Always remove aliases when an asset is removed

diff --git a/Engine/Plugins/Editor/ContentBrowser/ContentBrowserAliasDataSoure/Source/ContentBrowserAliasDataSource/Private/ContentBrowserAliasDataSource.cpp b/Engine/Plugins/Editor/ContentBrowser/ContentBrowserAliasDataSoure/Source/ContentBrowserAliasDataSource/Private/ContentBrowserAliasDataSource.cpp
--- a/Engine/Plugins/Editor/ContentBrowser/ContentBrowserAliasDataSoure/Source/ContentBrowserAliasDataSource/Private/ContentBrowserAliasDataSource.cpp
+++ b/Engine/Plugins/Editor/ContentBrowser/ContentBrowserAliasDataSoure/Source/ContentBrowserAliasDataSource/Private/ContentBrowserAliasDataSource.cpp
@@ -445,11 +445,7 @@ void UContentBrowserAliasDataSource::OnAssetAdded(const FAssetData& InAssetData)
 
 void UContentBrowserAliasDataSource::OnAssetRemoved(const FAssetData& InAssetData)
 {
-	FString AliasTagValue;
-	if (InAssetData.GetTagValue(AliasTagName, AliasTagValue))
-	{
-		RemoveAliases(InAssetData);
-	}
+	RemoveAliases(InAssetData);
 }
 
 void UContentBrowserAliasDataSource::ReconcileAliasesFromMetaData(const FAssetData& Asset)
@@ -625,6 +621,19 @@ bool UContentBrowserAliasDataSource::IsItemDirty(const FContentBrowserItemData&
 
 bool UContentBrowserAliasDataSource::CanEditItem(const FContentBrowserItemData& InItem, FText* OutErrorMsg)
 {
+	if (TSharedPtr<const FContentBrowserAliasItemDataPayload> AliasPayload = StaticCastSharedPtr<const FContentBrowserAliasItemDataPayload>(InItem.GetPayload()))
+	{
+		// Both the alias path and asset path must pass the writable folder filter in order to be editable
+		const TSharedRef<FPathPermissionList>& WritableFolderFilter = AssetTools->GetWritableFolderPermissionList();
+		if (!WritableFolderFilter->PassesStartsWithFilter(AliasPayload->Alias.Value))
+		{
+			if (OutErrorMsg)
+			{
+				*OutErrorMsg = FText::Format(NSLOCTEXT("ContentBrowserAliasDataSource", "Error_FolderIsLocked", "Alias '{0}' is in a locked folder"), FText::FromName(AliasPayload->Alias.Value));
+			}
+			return false;
+		}
+	}
 	return ContentBrowserAssetData::CanEditItem(AssetTools, this, InItem, OutErrorMsg);
 }
 
