Add null check in alias data source

diff --git a/Engine/Plugins/Editor/ContentBrowser/ContentBrowserAliasDataSoure/Source/ContentBrowserAliasDataSource/Private/ContentBrowserAliasDataSource.cpp b/Engine/Plugins/Editor/ContentBrowser/ContentBrowserAliasDataSoure/Source/ContentBrowserAliasDataSource/Private/ContentBrowserAliasDataSource.cpp
--- a/Engine/Plugins/Editor/ContentBrowser/ContentBrowserAliasDataSoure/Source/ContentBrowserAliasDataSource/Private/ContentBrowserAliasDataSource.cpp
+++ b/Engine/Plugins/Editor/ContentBrowser/ContentBrowserAliasDataSoure/Source/ContentBrowserAliasDataSource/Private/ContentBrowserAliasDataSource.cpp
@@ -253,7 +253,12 @@ bool UContentBrowserAliasDataSource::DoesItemPassFilter(const FContentBrowserIte
 		{
 			if (TSharedPtr<const FContentBrowserAliasItemDataPayload> AliasPayload = StaticCastSharedPtr<const FContentBrowserAliasItemDataPayload>(InItem.GetPayload()))
 			{
-				return DoesAliasPassFilter(AllAliases[AliasPayload->Alias], *AssetDataFilter);
+				FAliasData* FoundAlias = AllAliases.Find(AliasPayload->Alias);
+				ensureMsgf(FoundAlias != nullptr, TEXT("Alias %s->%s was not found in AllAliases, but it should have been"), *AliasPayload->Alias.Key.ToString(), *AliasPayload->Alias.Value.ToString());
+				if (FoundAlias)
+				{
+					return DoesAliasPassFilter(*FoundAlias, *AssetDataFilter);
+				}
 			}
 		}
 		break;
