Remove ensure from missing alias data, added comment to explain why this may legitimately happen.

diff --git a/Engine/Plugins/Editor/ContentBrowser/ContentBrowserAliasDataSoure/Source/ContentBrowserAliasDataSource/Private/ContentBrowserAliasDataSource.cpp b/Engine/Plugins/Editor/ContentBrowser/ContentBrowserAliasDataSoure/Source/ContentBrowserAliasDataSource/Private/ContentBrowserAliasDataSource.cpp
--- a/Engine/Plugins/Editor/ContentBrowser/ContentBrowserAliasDataSoure/Source/ContentBrowserAliasDataSource/Private/ContentBrowserAliasDataSource.cpp
+++ b/Engine/Plugins/Editor/ContentBrowser/ContentBrowserAliasDataSoure/Source/ContentBrowserAliasDataSource/Private/ContentBrowserAliasDataSource.cpp
@@ -254,7 +254,10 @@ bool UContentBrowserAliasDataSource::DoesItemPassFilter(const FContentBrowserIte
 			if (TSharedPtr<const FContentBrowserAliasItemDataPayload> AliasPayload = StaticCastSharedPtr<const FContentBrowserAliasItemDataPayload>(InItem.GetPayload()))
 			{
 				FAliasData* FoundAlias = AllAliases.Find(AliasPayload->Alias);
-				ensureMsgf(FoundAlias != nullptr, TEXT("Alias %s->%s was not found in AllAliases, but it should have been"), *AliasPayload->Alias.Key.ToString(), *AliasPayload->Alias.Value.ToString());
+				// This should always be true except in the case where the item is deleted in the same tick that DoesItemPassFilter is called.
+				// This is because the alias data source processes the deletion immediately, but the content browser deletion is queued until
+				// next tick, causing them to be briefly out of sync.
+				// An alternative solution would be to add a way to synchronously flush content browser updates, but that doesn't exist atm.
 				if (FoundAlias)
 				{
 					return DoesAliasPassFilter(*FoundAlias, *AssetDataFilter);
@@ -736,13 +739,17 @@ FContentBrowserItemData UContentBrowserAliasDataSource::CreateAssetFolderItem(co
 
 FContentBrowserItemData UContentBrowserAliasDataSource::CreateAssetFileItem(const FContentBrowserUniqueAlias& Alias)
 {
-	const FAliasData& AliasData = AllAliases[Alias];
-	
-	FName VirtualizedPath;
-	TryConvertInternalPathToVirtual(AliasData.AliasID, VirtualizedPath);
+	const FAliasData* AliasData = AllAliases.Find(Alias);
+	// See UContentBrowserAliasDataSource::DoesItemPassFilter for more information on how this could fail
+	if (AliasData)
+	{
+		FName VirtualizedPath;
+		TryConvertInternalPathToVirtual(AliasData->AliasID, VirtualizedPath);
 
-	// Since AliasID is PackagePath/AssetName, AssetName should also be passed as the ItemName here. This provides the functionality of
-	// being able to have multiple aliases with the same display name, while still showing their original asset name in the tooltip.
-	return FContentBrowserItemData(this, EContentBrowserItemFlags::Type_File | EContentBrowserItemFlags::Category_Asset, VirtualizedPath,
-		AliasData.AssetData.AssetName, FText::FromName(AliasData.AliasName), MakeShared<FContentBrowserAliasItemDataPayload>(AliasData.AssetData, Alias));
+		// Since AliasID is PackagePath/AssetName, AssetName should also be passed as the ItemName here. This provides the functionality of
+		// being able to have multiple aliases with the same display name, while still showing their original asset name in the tooltip.
+		return FContentBrowserItemData(this, EContentBrowserItemFlags::Type_File | EContentBrowserItemFlags::Category_Asset, VirtualizedPath,
+			AliasData->AssetData.AssetName, FText::FromName(AliasData->AliasName), MakeShared<FContentBrowserAliasItemDataPayload>(AliasData->AssetData, Alias));
+	}
+	return FContentBrowserItemData();
 }
