Add method to query the first alias found for a given asset

diff --git a/Engine/Plugins/Editor/ContentBrowser/ContentBrowserAliasDataSoure/Source/ContentBrowserAliasDataSource/Private/ContentBrowserAliasDataSource.cpp b/Engine/Plugins/Editor/ContentBrowser/ContentBrowserAliasDataSoure/Source/ContentBrowserAliasDataSource/Private/ContentBrowserAliasDataSource.cpp
--- a/Engine/Plugins/Editor/ContentBrowser/ContentBrowserAliasDataSoure/Source/ContentBrowserAliasDataSource/Private/ContentBrowserAliasDataSource.cpp
+++ b/Engine/Plugins/Editor/ContentBrowser/ContentBrowserAliasDataSoure/Source/ContentBrowserAliasDataSource/Private/ContentBrowserAliasDataSource.cpp
@@ -280,6 +280,18 @@ bool UContentBrowserAliasDataSource::DoesAliasPassFilter(const FAliasData& Alias
 		&& (Filter.ExclusiveFilter.IsEmpty() || !AssetRegistry->IsAssetIncludedByFilter(AliasAssetData, Filter.ExclusiveFilter)); // Passes Exclusive
 }
 
+FName UContentBrowserAliasDataSource::GetAliasForObjectPath(const FName ObjectPath) const
+{
+	if (const TArray<FName>* Aliases = AliasesForObjectPath.Find(ObjectPath))
+	{
+		if (Aliases->Num() > 0)
+		{
+			return (*Aliases)[0];
+		}
+	}
+	return NAME_None;
+}
+
 void UContentBrowserAliasDataSource::AddAliases(const FAssetData& Asset, const TArray<FName>& Aliases, bool bInIsFromMetaData)
 {
 	for (const FName Alias : Aliases)
diff --git a/Engine/Plugins/Editor/ContentBrowser/ContentBrowserAliasDataSoure/Source/ContentBrowserAliasDataSource/Public/ContentBrowserAliasDataSource.h b/Engine/Plugins/Editor/ContentBrowser/ContentBrowserAliasDataSoure/Source/ContentBrowserAliasDataSource/Public/ContentBrowserAliasDataSource.h
--- a/Engine/Plugins/Editor/ContentBrowser/ContentBrowserAliasDataSoure/Source/ContentBrowserAliasDataSource/Public/ContentBrowserAliasDataSource.h
+++ b/Engine/Plugins/Editor/ContentBrowser/ContentBrowserAliasDataSoure/Source/ContentBrowserAliasDataSource/Public/ContentBrowserAliasDataSource.h
@@ -85,6 +85,8 @@ public:
 	virtual bool Legacy_TryConvertAssetDataToVirtualPath(const FAssetData& InAssetData, const bool InUseFolderPaths, FName& OutPath) override;
 	// ~ End UContentBrowserDataSource interface
 
+	FName GetAliasForObjectPath(const FName ObjectPath) const;
+
 	/** Add a list of aliases for a given asset. bInIsFromMetaData should only be true if the list of aliases came from the AliasTagName metadata. */
 	void AddAliases(const FAssetData& Asset, const TArray<FName>& Aliases, bool bInIsFromMetaData = false);
 	/** Add an alias for a given asset. bInIsFromMetaData should only be true if the alias came from the AliasTagName metadata. */
