Fix missing include

diff --git a/Engine/Plugins/Editor/ContentBrowser/ContentBrowserAliasDataSoure/Source/ContentBrowserAliasDataSource/Public/ContentBrowserAliasDataSource.h b/Engine/Plugins/Editor/ContentBrowser/ContentBrowserAliasDataSoure/Source/ContentBrowserAliasDataSource/Public/ContentBrowserAliasDataSource.h
--- a/Engine/Plugins/Editor/ContentBrowser/ContentBrowserAliasDataSoure/Source/ContentBrowserAliasDataSource/Public/ContentBrowserAliasDataSource.h
+++ b/Engine/Plugins/Editor/ContentBrowser/ContentBrowserAliasDataSoure/Source/ContentBrowserAliasDataSource/Public/ContentBrowserAliasDataSource.h
@@ -10,6 +10,7 @@
 #include "AssetRegistry/AssetRegistryModule.h"
 #include "AssetRegistry/PathTree.h"
 #include "ContentBrowserItemData.h"
+#include "ContentBrowserItemPath.h"
 #include "ContentBrowserAliasDataSource.generated.h"
 
 DECLARE_LOG_CATEGORY_EXTERN(LogContentBrowserAliasDataSource, Log, All);
