delete Test variable in alias data source

diff --git a/Engine/Plugins/Editor/ContentBrowser/ContentBrowserAliasDataSoure/Source/ContentBrowserAliasDataSource/Public/ContentBrowserAliasDataSource.h b/Engine/Plugins/Editor/ContentBrowser/ContentBrowserAliasDataSoure/Source/ContentBrowserAliasDataSource/Public/ContentBrowserAliasDataSource.h
--- a/Engine/Plugins/Editor/ContentBrowser/ContentBrowserAliasDataSoure/Source/ContentBrowserAliasDataSource/Public/ContentBrowserAliasDataSource.h
+++ b/Engine/Plugins/Editor/ContentBrowser/ContentBrowserAliasDataSoure/Source/ContentBrowserAliasDataSource/Public/ContentBrowserAliasDataSource.h
@@ -159,7 +159,6 @@ private:
 
 	/** The full folder hierarchy for all alias paths */
 	FPathTree PathTree;
-	TMap<FContentBrowserUniqueAlias, FAliasData> Test;
 	/** Alias data keyed by their full alias path, ie /Game/MyData/Aliases/SourceMesh */
 	TMap<FContentBrowserUniqueAlias, FAliasData> AllAliases;
 	/** A list of alias paths to display for each asset, ie /Game/Meshes/SourceMesh.SourceMesh */
