FNames containing asset paths are deprecated. FSoftObjectPath or FTopLevelAssetPath should be used instead.
Fixups for editor support classes for content browser code.

diff --git a/Engine/Plugins/Editor/ContentBrowser/ContentBrowserAliasDataSoure/Source/ContentBrowserAliasDataSource/Private/ContentBrowserAliasDataSource.cpp b/Engine/Plugins/Editor/ContentBrowser/ContentBrowserAliasDataSoure/Source/ContentBrowserAliasDataSource/Private/ContentBrowserAliasDataSource.cpp
--- a/Engine/Plugins/Editor/ContentBrowser/ContentBrowserAliasDataSoure/Source/ContentBrowserAliasDataSource/Private/ContentBrowserAliasDataSource.cpp
+++ b/Engine/Plugins/Editor/ContentBrowser/ContentBrowserAliasDataSoure/Source/ContentBrowserAliasDataSource/Private/ContentBrowserAliasDataSource.cpp
@@ -144,7 +144,7 @@ void UContentBrowserAliasDataSource::EnumerateItemsMatchingFilter(const FContent
 						const FAliasData& AliasData = AllAliases[Alias];
 						if (DoesAliasPassFilter(AliasData, *AssetDataFilter))
 						{
-							AlreadyAddedOriginalAssets.Add(AliasData.AssetData.ObjectPath, &bIsAlreadyInSet);
+							AlreadyAddedOriginalAssets.Add(AliasData.AssetData.GetSoftObjectPath(), &bIsAlreadyInSet);
 							if (!bIsAlreadyInSet)
 							{
 								if (!InCallback(CreateAssetFileItem(Alias)))
@@ -165,7 +165,7 @@ void UContentBrowserAliasDataSource::EnumerateItemsMatchingFilter(const FContent
 			{
 				if (DoesAliasPassFilter(Pair.Value, *AssetDataFilter))
 				{
-					AlreadyAddedOriginalAssets.Add(Pair.Value.AssetData.ObjectPath, &bIsAlreadyInSet);
+					AlreadyAddedOriginalAssets.Add(Pair.Value.AssetData.GetSoftObjectPath(), &bIsAlreadyInSet);
 					if (!bIsAlreadyInSet)
 					{
 						if (!InCallback(CreateAssetFileItem(Pair.Key)))
@@ -212,7 +212,7 @@ void UContentBrowserAliasDataSource::EnumerateItemsAtPath(const FName InPath, co
 			{
 				for (const FContentBrowserUniqueAlias& Alias : *Aliases)
 				{
-					AlreadyAddedOriginalAssets.Add(AllAliases[Alias].AssetData.ObjectPath, &bIsAlreadyInSet);
+					AlreadyAddedOriginalAssets.Add(AllAliases[Alias].AssetData.GetSoftObjectPath(), &bIsAlreadyInSet);
 					if (!bIsAlreadyInSet)
 					{
 						if (!InCallback(CreateAssetFileItem(Alias)))
@@ -283,7 +283,7 @@ bool UContentBrowserAliasDataSource::DoesAliasPassFilter(const FAliasData& Alias
 		&& (Filter.ExclusiveFilter.IsEmpty() || !AssetRegistry->IsAssetIncludedByFilter(AliasAssetData, Filter.ExclusiveFilter)); // Passes Exclusive
 }
 
-TArray<FContentBrowserItemPath> UContentBrowserAliasDataSource::GetAliasesForPath(const FName ObjectPath) const
+TArray<FContentBrowserItemPath> UContentBrowserAliasDataSource::GetAliasesForPath(const FSoftObjectPath& ObjectPath) const
 {
 	TArray<FContentBrowserItemPath> OutAliases;
 
@@ -313,10 +313,10 @@ void UContentBrowserAliasDataSource::AddAlias(const FAssetData& Asset, const FNa
 {
 	auto LogErrorMessage = [&Alias, &Asset](const TCHAR* Reason)
 	{
-		UE_LOG(LogContentBrowserAliasDataSource, Warning, TEXT("Cannot add alias %s for %s because: %s"), *Alias.ToString(), *Asset.ObjectPath.ToString(), Reason);
+		UE_LOG(LogContentBrowserAliasDataSource, Warning, TEXT("Cannot add alias %s for %s because: %s"), *Alias.ToString(), *Asset.GetObjectPathString(), Reason);
 	};
 	
-	FContentBrowserUniqueAlias UniqueAlias = TPairInitializer<FName, FName>(Asset.ObjectPath, Alias);
+	FContentBrowserUniqueAlias UniqueAlias = TPairInitializer<FSoftObjectPath, FName>(Asset.GetSoftObjectPath(), Alias);
 	if (AllAliases.Contains(UniqueAlias))
 	{
 		LogErrorMessage(TEXT("An alias with that name already exists"));
@@ -368,7 +368,7 @@ void UContentBrowserAliasDataSource::AddAlias(const FAssetData& Asset, const FNa
 	});
 
 	AliasesInPackagePath.FindOrAdd(PackagePath).Add(UniqueAlias);
-	AliasesForObjectPath.FindOrAdd(Asset.ObjectPath).Add(Alias);
+	AliasesForObjectPath.FindOrAdd(Asset.GetSoftObjectPath()).Add(Alias);
 	AllAliases.Add(UniqueAlias, FAliasData(Asset, PackagePath, *Tokens.Last(), bInIsFromMetaData));
 	QueueItemDataUpdate(FContentBrowserItemDataUpdate::MakeItemAddedUpdate(CreateAssetFileItem(UniqueAlias)));
 
@@ -413,18 +413,25 @@ void UContentBrowserAliasDataSource::RemoveFoldersRecursively(FStringView LeafFo
 	}
 }
 
-void UContentBrowserAliasDataSource::RemoveAlias(const FName ObjectPath, const FName Alias)
+PRAGMA_DISABLE_DEPRECATION_WARNINGS
+void UContentBrowserAliasDataSource::RemoveAlias(FName ObjectPath, const FName Alias)
 {
-    FContentBrowserUniqueAlias UniqueAlias =  TPairInitializer<FName, FName>(ObjectPath, Alias);
+	RemoveAlias(FSoftObjectPath(ObjectPath), Alias);
+}
+PRAGMA_ENABLE_DEPRECATION_WARNINGS
+
+void UContentBrowserAliasDataSource::RemoveAlias(const FSoftObjectPath& ObjectPath, const FName Alias)
+{
+    FContentBrowserUniqueAlias UniqueAlias =  TPairInitializer<FSoftObjectPath, FName>(ObjectPath, Alias);
 	FAliasData AliasData;
 	// Store a copy of the item data before it's removed for the MakeItemRemovedUpdate notification
 	FContentBrowserItemData ItemData = CreateAssetFileItem(UniqueAlias);
 	if (AllAliases.RemoveAndCopyValue(UniqueAlias, AliasData))
 	{
-		AliasesForObjectPath[AliasData.AssetData.ObjectPath].Remove(Alias);
-		if (AliasesForObjectPath[AliasData.AssetData.ObjectPath].Num() == 0)
+		AliasesForObjectPath[AliasData.AssetData.GetSoftObjectPath()].Remove(Alias);
+		if (AliasesForObjectPath[AliasData.AssetData.GetSoftObjectPath()].Num() == 0)
 		{
-			AliasesForObjectPath.Remove(AliasData.AssetData.ObjectPath);
+			AliasesForObjectPath.Remove(AliasData.AssetData.GetSoftObjectPath());
 		}
 
 		AliasesInPackagePath[AliasData.PackagePath].Remove(UniqueAlias);
@@ -444,7 +451,14 @@ void UContentBrowserAliasDataSource::RemoveAlias(const FName ObjectPath, const F
 	}
 }
 
-void UContentBrowserAliasDataSource::RemoveAliases(const FName ObjectPath)
+PRAGMA_DISABLE_DEPRECATION_WARNINGS
+void UContentBrowserAliasDataSource::RemoveAliases(FName ObjectPath)
+{
+	RemoveAliases(FSoftObjectPath(ObjectPath));
+}
+PRAGMA_ENABLE_DEPRECATION_WARNINGS
+
+void UContentBrowserAliasDataSource::RemoveAliases(const FSoftObjectPath& ObjectPath)
 {
 	const TArray<FName>* AliasesPtr = AliasesForObjectPath.Find(ObjectPath);
 	if (AliasesPtr)
@@ -479,10 +493,10 @@ void UContentBrowserAliasDataSource::OnAssetRemoved(const FAssetData& InAssetDat
 
 void UContentBrowserAliasDataSource::ReconcileAliasesFromMetaData(const FAssetData& Asset)
 {
-	const TArray<FName>* ExistingAliasesPtr = AliasesForObjectPath.Find(Asset.ObjectPath);
+	const TArray<FName>* ExistingAliasesPtr = AliasesForObjectPath.Find(Asset.GetSoftObjectPath());
 	if (ExistingAliasesPtr)
 	{
-		FContentBrowserUniqueAlias UniqueAlias = TPairInitializer<FName, FName>(Asset.ObjectPath, NAME_None);
+		FContentBrowserUniqueAlias UniqueAlias = TPairInitializer<FSoftObjectPath, FName>(Asset.GetSoftObjectPath(), NAME_None);
 		
 		FString AliasTagValue;
 		if (Asset.GetTagValue(AliasTagName, AliasTagValue))
@@ -521,7 +535,7 @@ void UContentBrowserAliasDataSource::ReconcileAliasesFromMetaData(const FAssetDa
 
 			for (const FName ExistingAlias : AliasesOnlyInExisting)
 			{
-				RemoveAlias(Asset.ObjectPath, ExistingAlias);
+				RemoveAlias(Asset.GetSoftObjectPath(), ExistingAlias);
 			}
 			AddAliases(Asset, AliasesOnlyInNew, true);
 		}
@@ -534,7 +548,7 @@ void UContentBrowserAliasDataSource::ReconcileAliasesFromMetaData(const FAssetDa
 				UniqueAlias.Value = Alias;
 				if (AllAliases[UniqueAlias].bIsFromMetaData)
 				{
-					RemoveAlias(Asset.ObjectPath, Alias);
+					RemoveAlias(Asset.GetSoftObjectPath(), Alias);
 				}
 			}
 		}
@@ -550,11 +564,11 @@ void UContentBrowserAliasDataSource::ReconcileAliasesForAsset(const FAssetData&
 {
 	if (ensure(Asset.IsValid()))
 	{
-		const TArray<FName>* ExistingAliasesPtr = AliasesForObjectPath.Find(Asset.ObjectPath);
+		const TArray<FName>* ExistingAliasesPtr = AliasesForObjectPath.Find(Asset.GetSoftObjectPath());
 		if (ExistingAliasesPtr)
 		{
 			TArray<FName> AliasesOnlyInExisting, AliasesOnlyInNew;
-			FContentBrowserUniqueAlias UniqueAlias = TPairInitializer<FName, FName>(Asset.ObjectPath, NAME_None);
+			FContentBrowserUniqueAlias UniqueAlias = TPairInitializer<FSoftObjectPath, FName>(Asset.GetSoftObjectPath(), NAME_None);
 			for (const FName Alias : *ExistingAliasesPtr)
 			{
 				UniqueAlias.Value = Alias;
@@ -577,7 +591,7 @@ void UContentBrowserAliasDataSource::ReconcileAliasesForAsset(const FAssetData&
 
 			for (const FName ExistingAlias : AliasesOnlyInExisting)
 			{
-				RemoveAlias(Asset.ObjectPath, ExistingAlias);
+				RemoveAlias(Asset.GetSoftObjectPath(), ExistingAlias);
 			}
 			AddAliases(Asset, AliasesOnlyInNew);
 		}
@@ -595,9 +609,9 @@ void UContentBrowserAliasDataSource::OnAssetUpdated(const FAssetData& InAssetDat
 
 void UContentBrowserAliasDataSource::MakeItemModifiedUpdate(UObject* Object)
 {
-	const FName ObjectPath = *Object->GetPathName();
+	const FSoftObjectPath ObjectPath(Object);
 	const TArray<FName>* Aliases = AliasesForObjectPath.Find(ObjectPath);
-	FContentBrowserUniqueAlias UniqueAlias =  TPairInitializer<FName, FName>(ObjectPath, NAME_None);
+	FContentBrowserUniqueAlias UniqueAlias =  TPairInitializer<FSoftObjectPath, FName>(ObjectPath, NAME_None);
 	if (Aliases)
 	{
 		for (const FName Alias : *Aliases)
@@ -758,7 +772,9 @@ bool UContentBrowserAliasDataSource::Legacy_TryConvertPackagePathToVirtualPath(c
 bool UContentBrowserAliasDataSource::Legacy_TryConvertAssetDataToVirtualPath(const FAssetData& InAssetData, const bool InUseFolderPaths, FName& OutPath)
 {
 	return InAssetData.AssetClassPath != FTopLevelAssetPath(TEXT("/Script/CoreUObject"), NAME_Class) // Ignore legacy class items
+PRAGMA_DISABLE_DEPRECATION_WARNINGS
 		&& TryConvertInternalPathToVirtual(InUseFolderPaths ? InAssetData.PackagePath : InAssetData.ObjectPath, OutPath);
+PRAGMA_ENABLE_DEPRECATION_WARNINGS
 }
 
 FContentBrowserItemData UContentBrowserAliasDataSource::CreateAssetFolderItem(const FName InFolderPath)
diff --git a/Engine/Plugins/Editor/ContentBrowser/ContentBrowserAliasDataSoure/Source/ContentBrowserAliasDataSource/Public/ContentBrowserAliasDataSource.h b/Engine/Plugins/Editor/ContentBrowser/ContentBrowserAliasDataSoure/Source/ContentBrowserAliasDataSource/Public/ContentBrowserAliasDataSource.h
--- a/Engine/Plugins/Editor/ContentBrowser/ContentBrowserAliasDataSoure/Source/ContentBrowserAliasDataSource/Public/ContentBrowserAliasDataSource.h
+++ b/Engine/Plugins/Editor/ContentBrowser/ContentBrowserAliasDataSoure/Source/ContentBrowserAliasDataSource/Public/ContentBrowserAliasDataSource.h
@@ -19,7 +19,7 @@ class IAssetTools;
 struct FContentBrowserCompiledAssetDataFilter;
 
 /** A unique alias is a pair of SourceObjectPath:AliasPath, eg /Game/MyAsset.MyAsset:/Game/SomeFolder/MyAlias */
-typedef TPair<FName, FName> FContentBrowserUniqueAlias;
+typedef TPair<FSoftObjectPath, FName> FContentBrowserUniqueAlias;
 
 class CONTENTBROWSERALIASDATASOURCE_API FContentBrowserAliasItemDataPayload : public FContentBrowserAssetFileItemDataPayload
 {
@@ -59,7 +59,7 @@ public:
 
 	virtual bool DoesItemPassFilter(const FContentBrowserItemData& InItem, const FContentBrowserDataCompiledFilter& InFilter) override;
 
-	virtual TArray<FContentBrowserItemPath> GetAliasesForPath(const FName InInternalPath) const override;
+	virtual TArray<FContentBrowserItemPath> GetAliasesForPath(const FSoftObjectPath& InInternalPath) const override;
 
 	virtual bool GetItemAttribute(const FContentBrowserItemData& InItem, const bool InIncludeMetaData, const FName InAttributeKey, FContentBrowserItemDataAttributeValue& OutAttributeValue) override;
 	virtual bool GetItemAttributes(const FContentBrowserItemData& InItem, const bool InIncludeMetaData, FContentBrowserItemDataAttributeValues& OutAttributeValues) override;
@@ -95,11 +95,16 @@ public:
 	/** Add an alias for a given asset. bInIsFromMetaData should only be true if the alias came from the AliasTagName metadata. */
 	void AddAlias(const FAssetData& Asset, const FName Alias, bool bInIsFromMetaData = false);
 	/** Remove the given alias from the data source */
-	void RemoveAlias(const FName ObjectPath, const FName Alias);
+	void RemoveAlias(const FSoftObjectPath& ObjectPath, const FName Alias);
 	/** Remove all aliases for the given object */
-	void RemoveAliases(const FName ObjectPath);
+	void RemoveAliases(const FSoftObjectPath& ObjectPath);
 	/** Remove all aliases for the given asset */
-	void RemoveAliases(const FAssetData& Asset) { RemoveAliases(Asset.ObjectPath); }
+	void RemoveAliases(const FAssetData& Asset) { RemoveAliases(Asset.GetSoftObjectPath()); }
+
+	UE_DEPRECATED(5.1, "FNames containing full asset paths are deprecated, use FSoftObjectPath instead")
+	void RemoveAlias(const FName ObjectPath, const FName Alias);
+	UE_DEPRECATED(5.1, "FNames containing full asset paths are deprecated, use FSoftObjectPath instead")
+	void RemoveAliases(const FName ObjectPath);
 
 	/** Get all aliases from metadata for the given asset, then calls AddAlias or RemoveAlias for every alias that doesn't match the stored data. */
 	void ReconcileAliasesFromMetaData(const FAssetData& Asset);
@@ -162,9 +167,9 @@ private:
 	/** Alias data keyed by their full alias path, ie /Game/MyData/Aliases/SourceMesh */
 	TMap<FContentBrowserUniqueAlias, FAliasData> AllAliases;
 	/** A list of alias paths to display for each asset, ie /Game/Meshes/SourceMesh.SourceMesh */
-	TMap<FName, TArray<FName>> AliasesForObjectPath;
+	TMap<FSoftObjectPath, TArray<FName>> AliasesForObjectPath;
 	/** A list of alias paths to display for each folder, ie /Game/MyData/Aliases */
 	TMap<FName, TArray<FContentBrowserUniqueAlias>> AliasesInPackagePath;
 	/** A set used for removing duplicate aliases in the same query, stored here to avoid constant reallocation */
-	TSet<FName> AlreadyAddedOriginalAssets;
+	TSet<FSoftObjectPath> AlreadyAddedOriginalAssets;
 };
diff --git a/Engine/Plugins/Editor/ContentBrowser/ContentBrowserAssetDataSource/Source/ContentBrowserAssetDataSource/Private/AssetFileContextMenu.cpp b/Engine/Plugins/Editor/ContentBrowser/ContentBrowserAssetDataSource/Source/ContentBrowserAssetDataSource/Private/AssetFileContextMenu.cpp
--- a/Engine/Plugins/Editor/ContentBrowser/ContentBrowserAssetDataSource/Source/ContentBrowserAssetDataSource/Private/AssetFileContextMenu.cpp
+++ b/Engine/Plugins/Editor/ContentBrowser/ContentBrowserAssetDataSource/Source/ContentBrowserAssetDataSource/Private/AssetFileContextMenu.cpp
@@ -43,6 +43,7 @@
 #include "AssetViewUtils.h"
 #include "Dialogs/Dialogs.h"
 #include "SMetaDataView.h"
+#include "Algo/Transform.h"
 
 #include "ObjectTools.h"
 #include "PackageTools.h"
@@ -665,7 +666,7 @@ void FAssetFileContextMenu::MakeAssetLocalizationSubMenu(UToolMenu* Menu)
 
 		for (const FAssetData& Asset : SelectedAssets)
 		{
-			const FString AssetPath = Asset.ObjectPath.ToString();
+			const FString AssetPath = Asset.GetObjectPathString();
 
 			if (AssetViewUtils::IsEngineFolder(AssetPath))
 			{
@@ -744,19 +745,19 @@ void FAssetFileContextMenu::MakeAssetLocalizationSubMenu(UToolMenu* Menu)
 				continue;
 			}
 
-			const FString ObjectPath = Asset.ObjectPath.ToString();
+			const FString ObjectPath = Asset.GetObjectPathString();
 			if (FPackageName::IsLocalizedPackage(ObjectPath))
 			{
 				// Get the source path for this asset
 				FString SourceObjectPath;
 				if (FPackageLocalizationUtil::ConvertLocalizedToSource(ObjectPath, SourceObjectPath))
 				{
-					SourceAssetsState.CurrentAssets.Add(*SourceObjectPath);
+					SourceAssetsState.CurrentAssets.Add(FSoftObjectPath(SourceObjectPath));
 				}
 			}
 			else
 			{
-				SourceAssetsState.SelectedAssets.Add(Asset.ObjectPath);
+				SourceAssetsState.SelectedAssets.Add(Asset.GetSoftObjectPath());
 
 				// Get the localized path for this asset and culture
 				FString LocalizedObjectPath;
@@ -764,15 +765,15 @@ void FAssetFileContextMenu::MakeAssetLocalizationSubMenu(UToolMenu* Menu)
 				{
 					// Does this localized asset already exist?
 					FAssetRegistryModule& AssetRegistryModule = FModuleManager::LoadModuleChecked<FAssetRegistryModule>(TEXT("AssetRegistry"));
-					FAssetData LocalizedAssetData = AssetRegistryModule.Get().GetAssetByObjectPath(*LocalizedObjectPath);
+					FAssetData LocalizedAssetData = AssetRegistryModule.Get().GetAssetByObjectPath(FSoftObjectPath(LocalizedObjectPath));
 
 					if (LocalizedAssetData.IsValid())
 					{
-						LocalizedAssetsStateForCulture.CurrentAssets.Add(*LocalizedObjectPath);
+						LocalizedAssetsStateForCulture.CurrentAssets.Add(FSoftObjectPath(LocalizedObjectPath));
 					}
 					else
 					{
-						LocalizedAssetsStateForCulture.NewAssets.Add(*LocalizedObjectPath);
+						LocalizedAssetsStateForCulture.NewAssets.Add(FSoftObjectPath(LocalizedObjectPath));
 					}
 				}
 			}
@@ -839,7 +840,9 @@ void FAssetFileContextMenu::MakeAssetLocalizationSubMenu(UToolMenu* Menu)
 				LOCTEXT("ShowSourceAsset", "Show Source Asset"),
 				LOCTEXT("ShowSourceAssetTooltip", "Show the source asset in the Content Browser."),
 				FSlateIcon(FAppStyle::GetAppStyleSetName(), "SystemWideCommands.FindInContentBrowser"),
+PRAGMA_DISABLE_DEPRECATION_WARNINGS
 				FUIAction(FExecuteAction::CreateSP(this, &FAssetFileContextMenu::ExecuteFindInAssetTree, SourceAssetsState.CurrentAssets.Array()))
+PRAGMA_ENABLE_DEPRECATION_WARNINGS
 				);
 
 			Section.AddMenuEntry(
@@ -847,7 +850,9 @@ void FAssetFileContextMenu::MakeAssetLocalizationSubMenu(UToolMenu* Menu)
 				LOCTEXT("EditSourceAsset", "Edit Source Asset"),
 				LOCTEXT("EditSourceAssetTooltip", "Edit the source asset."),
 				FSlateIcon(FAppStyle::GetAppStyleSetName(), "ContentBrowser.AssetActions.Edit"),
+PRAGMA_DISABLE_DEPRECATION_WARNINGS
 				FUIAction(FExecuteAction::CreateSP(this, &FAssetFileContextMenu::ExecuteOpenEditorsForAssets, SourceAssetsState.CurrentAssets.Array()))
+PRAGMA_ENABLE_DEPRECATION_WARNINGS
 				);
 		}
 	}
@@ -896,7 +901,7 @@ void FAssetFileContextMenu::MakeAssetLocalizationSubMenu(UToolMenu* Menu)
 	}
 }
 
-void FAssetFileContextMenu::MakeCreateLocalizedAssetSubMenu(UToolMenu* Menu, TSet<FName> InSelectedSourceAssets, TArray<FLocalizedAssetsState> InLocalizedAssetsState)
+void FAssetFileContextMenu::MakeCreateLocalizedAssetSubMenu(UToolMenu* Menu, TSet<FSoftObjectPath> InSelectedSourceAssets, TArray<FLocalizedAssetsState> InLocalizedAssetsState)
 {
 	FToolMenuSection& Section = Menu->AddSection("Section");
 
@@ -929,7 +934,9 @@ void FAssetFileContextMenu::MakeShowLocalizedAssetSubMenu(UToolMenu* Menu, TArra
 				FText::FromString(LocalizedAssetsStateForCulture.Culture->GetDisplayName()),
 				FText::GetEmpty(),
 				FSlateIcon(),
+PRAGMA_DISABLE_DEPRECATION_WARNINGS
 				FUIAction(FExecuteAction::CreateSP(this, &FAssetFileContextMenu::ExecuteFindInAssetTree, LocalizedAssetsStateForCulture.CurrentAssets.Array()))
+PRAGMA_ENABLE_DEPRECATION_WARNINGS
 				);
 		}
 	}
@@ -948,26 +955,28 @@ void FAssetFileContextMenu::MakeEditLocalizedAssetSubMenu(UToolMenu* Menu, TArra
 				FText::FromString(LocalizedAssetsStateForCulture.Culture->GetDisplayName()),
 				FText::GetEmpty(),
 				FSlateIcon(),
+PRAGMA_DISABLE_DEPRECATION_WARNINGS
 				FUIAction(FExecuteAction::CreateSP(this, &FAssetFileContextMenu::ExecuteOpenEditorsForAssets, LocalizedAssetsStateForCulture.CurrentAssets.Array()))
+PRAGMA_ENABLE_DEPRECATION_WARNINGS
 				);
 		}
 	}
 }
 
-void FAssetFileContextMenu::ExecuteCreateLocalizedAsset(TSet<FName> InSelectedSourceAssets, FLocalizedAssetsState InLocalizedAssetsStateForCulture)
+void FAssetFileContextMenu::ExecuteCreateLocalizedAsset(TSet<FSoftObjectPath> InSelectedSourceAssets, FLocalizedAssetsState InLocalizedAssetsStateForCulture)
 {
 	TArray<UPackage*> PackagesToSave;
 	TArray<FAssetData> NewObjects;
 
-	for (const FName& SourceAssetName : InSelectedSourceAssets)
+	for (const FSoftObjectPath& SourceAsset: InSelectedSourceAssets)
 	{
-		if (InLocalizedAssetsStateForCulture.CurrentAssets.Contains(SourceAssetName))
+		if (InLocalizedAssetsStateForCulture.CurrentAssets.Contains(SourceAsset))
 		{
 			// Asset is already localized
 			continue;
 		}
 
-		UObject* SourceAssetObject = LoadObject<UObject>(nullptr, *SourceAssetName.ToString());
+		UObject* SourceAssetObject = LoadObject<UObject>(nullptr, *SourceAsset.ToString());
 		if (!SourceAssetObject)
 		{
 			// Source object cannot be loaded
@@ -1001,12 +1010,12 @@ void FAssetFileContextMenu::ExecuteCreateLocalizedAsset(TSet<FName> InSelectedSo
 	OnShowAssetsInPathsView.ExecuteIfBound(NewObjects);
 }
 
-void FAssetFileContextMenu::ExecuteFindInAssetTree(TArray<FName> InAssets)
+void FAssetFileContextMenu::ExecuteFindInAssetTree(TArray<FSoftObjectPath> InAssets)
 {
 	FAssetRegistryModule& AssetRegistryModule = FModuleManager::LoadModuleChecked<FAssetRegistryModule>(TEXT("AssetRegistry"));
 
 	FARFilter ARFilter;
-	ARFilter.ObjectPaths = MoveTemp(InAssets);
+	ARFilter.SoftObjectPaths = InAssets;
 	
 	TArray<FAssetData> FoundLocalizedAssetData;
 	AssetRegistryModule.Get().GetAssets(ARFilter, FoundLocalizedAssetData);
@@ -1014,7 +1023,7 @@ void FAssetFileContextMenu::ExecuteFindInAssetTree(TArray<FName> InAssets)
 	OnShowAssetsInPathsView.ExecuteIfBound(FoundLocalizedAssetData);
 }
 
-void FAssetFileContextMenu::ExecuteOpenEditorsForAssets(TArray<FName> InAssets)
+void FAssetFileContextMenu::ExecuteOpenEditorsForAssets(TArray<FSoftObjectPath> InAssets)
 {
 	GEditor->GetEditorSubsystem<UAssetEditorSubsystem>()->OpenEditorsForAssets(InAssets);
 }
diff --git a/Engine/Plugins/Editor/ContentBrowser/ContentBrowserAssetDataSource/Source/ContentBrowserAssetDataSource/Private/AssetFolderContextMenu.cpp b/Engine/Plugins/Editor/ContentBrowser/ContentBrowserAssetDataSource/Source/ContentBrowserAssetDataSource/Private/AssetFolderContextMenu.cpp
--- a/Engine/Plugins/Editor/ContentBrowser/ContentBrowserAssetDataSource/Source/ContentBrowserAssetDataSource/Private/AssetFolderContextMenu.cpp
+++ b/Engine/Plugins/Editor/ContentBrowser/ContentBrowserAssetDataSource/Source/ContentBrowserAssetDataSource/Private/AssetFolderContextMenu.cpp
@@ -236,7 +236,7 @@ void FAssetFolderContextMenu::ExecuteFixUpRedirectorsInFolder()
 		TArray<FString> ObjectPaths;
 		for (const auto& Asset : AssetList)
 		{
-			ObjectPaths.Add(Asset.ObjectPath.ToString());
+			ObjectPaths.Add(Asset.GetObjectPathString());
 		}
 
 		TArray<UObject*> Objects;
diff --git a/Engine/Plugins/Editor/ContentBrowser/ContentBrowserAssetDataSource/Source/ContentBrowserAssetDataSource/Private/ContentBrowserAssetDataCore.cpp b/Engine/Plugins/Editor/ContentBrowser/ContentBrowserAssetDataSource/Source/ContentBrowserAssetDataSource/Private/ContentBrowserAssetDataCore.cpp
--- a/Engine/Plugins/Editor/ContentBrowser/ContentBrowserAssetDataSource/Source/ContentBrowserAssetDataSource/Private/ContentBrowserAssetDataCore.cpp
+++ b/Engine/Plugins/Editor/ContentBrowser/ContentBrowserAssetDataSource/Source/ContentBrowserAssetDataSource/Private/ContentBrowserAssetDataCore.cpp
@@ -265,7 +265,7 @@ bool EditOrPreviewAssetFileItems(TArrayView<const TSharedRef<const FContentBrows
 		
 		for (const FAssetData& AssetData : AssetsToLoad)
 		{
-			if (!AssetData.IsAssetLoaded() && FEditorFileUtils::IsMapPackageAsset(AssetData.ObjectPath.ToString()))
+			if (!AssetData.IsAssetLoaded() && FEditorFileUtils::IsMapPackageAsset(AssetData.GetObjectPathString()))
 			{
 				SlowTask.MakeDialog();
 			}
diff --git a/Engine/Plugins/Editor/ContentBrowser/ContentBrowserAssetDataSource/Source/ContentBrowserAssetDataSource/Private/ContentBrowserAssetDataPayload.cpp b/Engine/Plugins/Editor/ContentBrowser/ContentBrowserAssetDataSource/Source/ContentBrowserAssetDataSource/Private/ContentBrowserAssetDataPayload.cpp
--- a/Engine/Plugins/Editor/ContentBrowser/ContentBrowserAssetDataSource/Source/ContentBrowserAssetDataSource/Private/ContentBrowserAssetDataPayload.cpp
+++ b/Engine/Plugins/Editor/ContentBrowser/ContentBrowserAssetDataSource/Source/ContentBrowserAssetDataSource/Private/ContentBrowserAssetDataPayload.cpp
@@ -70,9 +70,9 @@ UObject* FContentBrowserAssetFileItemDataPayload::GetAsset(const bool bTryRecach
 	
 	if (!bHasCachedAssetPtr || (bTryRecacheIfNull && !CachedAssetPtr.IsValid()))
 	{
-		if (!AssetData.ObjectPath.IsNone())
+		if (AssetData.IsValid())
 		{
-			CachedAssetPtr = FindObjectSafe<UObject>(nullptr, *FNameBuilder(AssetData.ObjectPath));
+			CachedAssetPtr = FindObjectSafe<UObject>(nullptr, *AssetData.GetObjectPathString());
 		}
 		bHasCachedAssetPtr = true;
 	}
@@ -83,10 +83,10 @@ UObject* FContentBrowserAssetFileItemDataPayload::LoadAsset(TSet<FName> LoadTags
 {
 	if (!bHasCachedAssetPtr || !CachedAssetPtr.IsValid())
 	{
-		if (!AssetData.ObjectPath.IsNone())
+		if (AssetData.IsValid())
 		{
 			FLinkerInstancingContext InstancingContext(MoveTemp(LoadTags));
-			CachedAssetPtr = LoadObject<UObject>(nullptr, *FNameBuilder(AssetData.ObjectPath), nullptr, LOAD_None, nullptr, &InstancingContext);
+			CachedAssetPtr = LoadObject<UObject>(nullptr, *AssetData.GetObjectPathString(), nullptr, LOAD_None, nullptr, &InstancingContext);
 			(void)GetPackage(/*bTryRecacheIfNull*/true); // Also re-cache the package pointer
 		}
 		bHasCachedAssetPtr = true;
diff --git a/Engine/Plugins/Editor/ContentBrowser/ContentBrowserAssetDataSource/Source/ContentBrowserAssetDataSource/Private/ContentBrowserAssetDataSource.cpp b/Engine/Plugins/Editor/ContentBrowser/ContentBrowserAssetDataSource/Source/ContentBrowserAssetDataSource/Private/ContentBrowserAssetDataSource.cpp
--- a/Engine/Plugins/Editor/ContentBrowser/ContentBrowserAssetDataSource/Source/ContentBrowserAssetDataSource/Private/ContentBrowserAssetDataSource.cpp
+++ b/Engine/Plugins/Editor/ContentBrowser/ContentBrowserAssetDataSource/Source/ContentBrowserAssetDataSource/Private/ContentBrowserAssetDataSource.cpp
@@ -427,7 +427,9 @@ bool UContentBrowserAssetDataSource::CreateAssetFilter(FAssetFilterInputParams&
 			FARFilter InclusiveFilter;
 			if (Params.ObjectFilter)
 			{
+PRAGMA_DISABLE_DEPRECATION_WARNINGS
 				InclusiveFilter.ObjectPaths.Append(Params.ObjectFilter->ObjectNamesToInclude);
+PRAGMA_ENABLE_DEPRECATION_WARNINGS
 				InclusiveFilter.TagsAndValues.Append(Params.ObjectFilter->TagsAndValuesToInclude);
 				InclusiveFilter.bIncludeOnlyOnDiskAssets |= Params.ObjectFilter->bOnDiskObjectsOnly;
 			}
@@ -444,13 +446,13 @@ bool UContentBrowserAssetDataSource::CreateAssetFilter(FAssetFilterInputParams&
 			}
 			if (Params.CollectionFilter)
 			{
-				TArray<FName> ObjectPathsForCollections;
+				TArray<FSoftObjectPath> ObjectPathsForCollections;
 				if (GetObjectPathsForCollections(Params.CollectionManager, Params.CollectionFilter->SelectedCollections, Params.CollectionFilter->bIncludeChildCollections, ObjectPathsForCollections) && ObjectPathsForCollections.Num() == 0)
 				{
 					// If we had collections but they contained no objects then we can bail as nothing will pass the filter
 					return false;
 				}
-				InclusiveFilter.ObjectPaths.Append(MoveTemp(ObjectPathsForCollections));
+				InclusiveFilter.SoftObjectPaths.Append(MoveTemp(ObjectPathsForCollections));
 			}
 
 #if DO_ENSURE
@@ -583,7 +585,9 @@ bool UContentBrowserAssetDataSource::CreateAssetFilter(FAssetFilterInputParams&
 			FARFilter ExclusiveFilter;
 			if (Params.ObjectFilter)
 			{
+PRAGMA_DISABLE_DEPRECATION_WARNINGS
 				ExclusiveFilter.ObjectPaths.Append(Params.ObjectFilter->ObjectNamesToExclude);
+PRAGMA_ENABLE_DEPRECATION_WARNINGS
 				ExclusiveFilter.TagsAndValues.Append(Params.ObjectFilter->TagsAndValuesToExclude);
 				ExclusiveFilter.bIncludeOnlyOnDiskAssets |= Params.ObjectFilter->bOnDiskObjectsOnly;
 			}
@@ -657,14 +661,14 @@ bool UContentBrowserAssetDataSource::CreateAssetFilter(FAssetFilterInputParams&
 			}
 			CompiledExclusiveFilter.PackagePaths.Reset();
 		}
-		if (CompiledInclusiveFilter.ObjectPaths.Num() > 0 && CompiledExclusiveFilter.ObjectPaths.Num() > 0)
+		if (CompiledInclusiveFilter.SoftObjectPaths.Num() > 0 && CompiledExclusiveFilter.SoftObjectPaths.Num() > 0)
 		{
-			CompiledInclusiveFilter.ObjectPaths = CompiledInclusiveFilter.ObjectPaths.Difference(CompiledExclusiveFilter.ObjectPaths);
-			if (CompiledInclusiveFilter.ObjectPaths.Num() == 0)
+			CompiledInclusiveFilter.SoftObjectPaths = CompiledInclusiveFilter.SoftObjectPaths.Difference(CompiledExclusiveFilter.SoftObjectPaths);
+			if (CompiledInclusiveFilter.SoftObjectPaths.Num() == 0)
 			{
 				return false;
 			}
-			CompiledExclusiveFilter.ObjectPaths.Reset();
+			CompiledExclusiveFilter.SoftObjectPaths.Reset();
 		}
 		if (CompiledInclusiveFilter.ClassPaths.Num() > 0 && CompiledExclusiveFilter.ClassPaths.Num() > 0)
 		{
@@ -738,7 +742,9 @@ void UContentBrowserAssetDataSource::CompileFilter(const FName InPath, const FCo
 						FARFilter CustomSourceAssetsFilter;
 						CustomSourceAssetsFilter.PackageNames = Params.AssetDataFilter->InclusiveFilter.PackageNames.Array();
 						CustomSourceAssetsFilter.PackagePaths = Params.AssetDataFilter->InclusiveFilter.PackagePaths.Array();
+PRAGMA_DISABLE_DEPRECATION_WARNINGS
 						CustomSourceAssetsFilter.ObjectPaths = Params.AssetDataFilter->InclusiveFilter.ObjectPaths.Array();
+PRAGMA_ENABLE_DEPRECATION_WARNINGS
 						CustomSourceAssetsFilter.ClassPaths = Params.AssetDataFilter->InclusiveFilter.ClassPaths.Array();
 						CustomSourceAssetsFilter.TagsAndValues = Params.AssetDataFilter->InclusiveFilter.TagsAndValues;
 						CustomSourceAssetsFilter.bIncludeOnlyOnDiskAssets = Params.AssetDataFilter->InclusiveFilter.bIncludeOnlyOnDiskAssets;
@@ -930,7 +936,9 @@ void UContentBrowserAssetDataSource::EnumerateItemsAtPath(const FName InPath, co
 	if (EnumHasAnyFlags(InItemTypeFilter, EContentBrowserItemTypeFilter::IncludeFiles))
 	{
 		FARFilter ARFilter;
+PRAGMA_DISABLE_DEPRECATION_WARNINGS
 		ARFilter.ObjectPaths.Add(InternalPath);
+PRAGMA_ENABLE_DEPRECATION_WARNINGS
 		AssetRegistry->EnumerateAssets(ARFilter, [this, &InCallback](const FAssetData& AssetData)
 		{
 			if (ContentBrowserAssetData::IsPrimaryAsset(AssetData))
@@ -1254,7 +1262,9 @@ bool UContentBrowserAssetDataSource::DuplicateItem(const FContentBrowserItemData
 	if (ContentBrowserAssetData::DuplicateItem(AssetTools, this, InItem, SourceAsset, NewAssetData))
 	{
 		FName VirtualizedPath;
+PRAGMA_DISABLE_DEPRECATION_WARNINGS
 		TryConvertInternalPathToVirtual(NewAssetData.ObjectPath, VirtualizedPath);
+PRAGMA_ENABLE_DEPRECATION_WARNINGS
 
 		FContentBrowserItemData NewItemData(
 			this,
@@ -1575,7 +1585,9 @@ bool UContentBrowserAssetDataSource::TryGetCollectionId(const FContentBrowserIte
 {
 	if (TSharedPtr<const FContentBrowserAssetFileItemDataPayload> AssetPayload = GetAssetFileItemPayload(InItem))
 	{
+PRAGMA_DISABLE_DEPRECATION_WARNINGS
 		OutCollectionId = AssetPayload->GetAssetData().ObjectPath;
+PRAGMA_ENABLE_DEPRECATION_WARNINGS
 		return true;
 	}
 	return false;
@@ -1610,7 +1622,9 @@ bool UContentBrowserAssetDataSource::Legacy_TryConvertPackagePathToVirtualPath(c
 bool UContentBrowserAssetDataSource::Legacy_TryConvertAssetDataToVirtualPath(const FAssetData& InAssetData, const bool InUseFolderPaths, FName& OutPath)
 {
 	return InAssetData.AssetClassPath != FTopLevelAssetPath(TEXT("/Script/CoreUObject"), TEXT("Class")) // Ignore legacy class items
+PRAGMA_DISABLE_DEPRECATION_WARNINGS
 		&& TryConvertInternalPathToVirtual(InUseFolderPaths ? InAssetData.PackagePath : InAssetData.ObjectPath, OutPath);
+PRAGMA_ENABLE_DEPRECATION_WARNINGS
 }
 
 bool UContentBrowserAssetDataSource::IsKnownContentPath(const FName InPackagePath) const
@@ -1641,15 +1655,20 @@ bool UContentBrowserAssetDataSource::IsRootContentPath(const FName InPackagePath
 	});
 }
 
-bool UContentBrowserAssetDataSource::GetObjectPathsForCollections(ICollectionManager* CollectionManager, TArrayView<const FCollectionNameType> InCollections, const bool bIncludeChildCollections, TArray<FName>& OutObjectPaths)
+bool UContentBrowserAssetDataSource::GetObjectPathsForCollections(ICollectionManager* CollectionManager, TArrayView<const FCollectionNameType> InCollections, const bool bIncludeChildCollections, TArray<FSoftObjectPath>& OutObjectPaths)
 {
 	if (InCollections.Num() > 0)
 	{
 		const ECollectionRecursionFlags::Flags CollectionRecursionMode = bIncludeChildCollections ? ECollectionRecursionFlags::SelfAndChildren : ECollectionRecursionFlags::Self;
 		
+		TArray<FName> TempNames;
 		for (const FCollectionNameType& CollectionNameType : InCollections)
 		{
-			CollectionManager->GetObjectsInCollection(CollectionNameType.Name, CollectionNameType.Type, OutObjectPaths, CollectionRecursionMode);
+			TempNames.Reset();
+			CollectionManager->GetObjectsInCollection(CollectionNameType.Name, CollectionNameType.Type, TempNames, CollectionRecursionMode);
+PRAGMA_DISABLE_DEPRECATION_WARNINGS
+			OutObjectPaths.Append(UE::SoftObjectPath::Private::ConvertObjectPathNames(TempNames));
+PRAGMA_ENABLE_DEPRECATION_WARNINGS
 		}
 
 		return true;
@@ -1669,7 +1688,9 @@ FContentBrowserItemData UContentBrowserAssetDataSource::CreateAssetFolderItem(co
 FContentBrowserItemData UContentBrowserAssetDataSource::CreateAssetFileItem(const FAssetData& InAssetData)
 {
 	FName VirtualizedPath;
+PRAGMA_DISABLE_DEPRECATION_WARNINGS
 	TryConvertInternalPathToVirtual(InAssetData.ObjectPath, VirtualizedPath);
+PRAGMA_ENABLE_DEPRECATION_WARNINGS
 
 	return ContentBrowserAssetData::CreateAssetFileItem(this, VirtualizedPath, InAssetData);
 }
@@ -2107,7 +2128,9 @@ void UContentBrowserAssetDataSource::OnBeginCreateAsset(const FName InDefaultAss
 		FAssetData NewAssetData(*(InPackagePath.ToString() / InDefaultAssetName.ToString()), InPackagePath, InDefaultAssetName, ClassToUse->GetClassPathName());
 
 		FName VirtualizedPath;
+PRAGMA_DISABLE_DEPRECATION_WARNINGS
 		TryConvertInternalPathToVirtual(NewAssetData.ObjectPath, VirtualizedPath);
+PRAGMA_ENABLE_DEPRECATION_WARNINGS
 
 		FContentBrowserItemData NewItemData(
 			this, 
diff --git a/Engine/Plugins/Editor/ContentBrowser/ContentBrowserAssetDataSource/Source/ContentBrowserAssetDataSource/Public/AssetFileContextMenu.h b/Engine/Plugins/Editor/ContentBrowser/ContentBrowserAssetDataSource/Source/ContentBrowserAssetDataSource/Public/AssetFileContextMenu.h
--- a/Engine/Plugins/Editor/ContentBrowser/ContentBrowserAssetDataSource/Source/ContentBrowserAssetDataSource/Public/AssetFileContextMenu.h
+++ b/Engine/Plugins/Editor/ContentBrowser/ContentBrowserAssetDataSource/Source/ContentBrowserAssetDataSource/Public/AssetFileContextMenu.h
@@ -24,15 +24,15 @@ public:
 private:
 	struct FSourceAssetsState
 	{
-		TSet<FName> SelectedAssets;
-		TSet<FName> CurrentAssets;
+		TSet<FSoftObjectPath> SelectedAssets;
+		TSet<FSoftObjectPath> CurrentAssets;
 	};
 
 	struct FLocalizedAssetsState
 	{
 		FCulturePtr Culture;
-		TSet<FName> NewAssets;
-		TSet<FName> CurrentAssets;
+		TSet<FSoftObjectPath> NewAssets;
+		TSet<FSoftObjectPath> CurrentAssets;
 	};
 
 private:
@@ -90,7 +90,7 @@ private:
 	void MakeAssetLocalizationSubMenu(UToolMenu* Menu);
 
 	/** Adds the Create Localized Asset sub-menu to a menu builder. */
-	void MakeCreateLocalizedAssetSubMenu(UToolMenu* Menu, TSet<FName> InSelectedSourceAssets, TArray<FLocalizedAssetsState> InLocalizedAssetsState);
+	void MakeCreateLocalizedAssetSubMenu(UToolMenu* Menu, TSet<FSoftObjectPath> InSelectedSourceAssets, TArray<FLocalizedAssetsState> InLocalizedAssetsState);
 
 	/** Adds the Show Localized Assets sub-menu to a menu builder. */
 	void MakeShowLocalizedAssetSubMenu(UToolMenu* Menu, TArray<FLocalizedAssetsState> InLocalizedAssetsState);
@@ -99,13 +99,13 @@ private:
 	void MakeEditLocalizedAssetSubMenu(UToolMenu* Menu, TArray<FLocalizedAssetsState> InLocalizedAssetsState);
 
 	/** Create new localized assets for the given culture */
-	void ExecuteCreateLocalizedAsset(TSet<FName> InSelectedSourceAssets, FLocalizedAssetsState InLocalizedAssetsStateForCulture);
+	void ExecuteCreateLocalizedAsset(TSet<FSoftObjectPath> InSelectedSourceAssets, FLocalizedAssetsState InLocalizedAssetsStateForCulture);
 
 	/** Find the given assets in the Content Browser */
-	void ExecuteFindInAssetTree(TArray<FName> InAssets);
+	void ExecuteFindInAssetTree(TArray<FSoftObjectPath> InAssets);
 
 	/** Open the given assets in their respective editors */
-	void ExecuteOpenEditorsForAssets(TArray<FName> InAssets);
+	void ExecuteOpenEditorsForAssets(TArray<FSoftObjectPath> InAssets);
 
 	/** Adds asset documentation menu options to a menu builder. Returns true if any options were added. */
 	bool AddDocumentationMenuOptions(UToolMenu* Menu);
diff --git a/Engine/Plugins/Editor/ContentBrowser/ContentBrowserAssetDataSource/Source/ContentBrowserAssetDataSource/Public/ContentBrowserAssetDataSource.h b/Engine/Plugins/Editor/ContentBrowser/ContentBrowserAssetDataSource/Source/ContentBrowserAssetDataSource/Public/ContentBrowserAssetDataSource.h
--- a/Engine/Plugins/Editor/ContentBrowser/ContentBrowserAssetDataSource/Source/ContentBrowserAssetDataSource/Public/ContentBrowserAssetDataSource.h
+++ b/Engine/Plugins/Editor/ContentBrowser/ContentBrowserAssetDataSource/Source/ContentBrowserAssetDataSource/Public/ContentBrowserAssetDataSource.h
@@ -266,7 +266,7 @@ private:
 
 	bool IsRootContentPath(const FName InPackagePath) const;
 
-	static bool GetObjectPathsForCollections(ICollectionManager* CollectionManager, TArrayView<const FCollectionNameType> InCollections, const bool bIncludeChildCollections, TArray<FName>& OutObjectPaths);
+	static bool GetObjectPathsForCollections(ICollectionManager* CollectionManager, TArrayView<const FCollectionNameType> InCollections, const bool bIncludeChildCollections, TArray<FSoftObjectPath>& OutObjectPaths);
 
 	FContentBrowserItemData CreateAssetFolderItem(const FName InFolderPath);
 
diff --git a/Engine/Plugins/Editor/ContentBrowser/ContentBrowserClassDataSource/Source/ContentBrowserClassDataSource/Private/ContentBrowserClassDataSource.cpp b/Engine/Plugins/Editor/ContentBrowser/ContentBrowserClassDataSource/Source/ContentBrowserClassDataSource/Private/ContentBrowserClassDataSource.cpp
--- a/Engine/Plugins/Editor/ContentBrowser/ContentBrowserClassDataSource/Source/ContentBrowserClassDataSource/Private/ContentBrowserClassDataSource.cpp
+++ b/Engine/Plugins/Editor/ContentBrowser/ContentBrowserClassDataSource/Source/ContentBrowserClassDataSource/Private/ContentBrowserClassDataSource.cpp
@@ -516,7 +516,9 @@ bool UContentBrowserClassDataSource::TryGetCollectionId(const FContentBrowserIte
 {
 	if (TSharedPtr<const FContentBrowserClassFileItemDataPayload> ClassPayload = GetClassFileItemPayload(InItem))
 	{
+PRAGMA_DISABLE_DEPRECATION_WARNINGS
 		OutCollectionId = ClassPayload->GetAssetData().ObjectPath;
+PRAGMA_ENABLE_DEPRECATION_WARNINGS
 		return true;
 	}
 	return false;
@@ -551,7 +553,9 @@ bool UContentBrowserClassDataSource::Legacy_TryConvertPackagePathToVirtualPath(c
 bool UContentBrowserClassDataSource::Legacy_TryConvertAssetDataToVirtualPath(const FAssetData& InAssetData, const bool InUseFolderPaths, FName& OutPath)
 {
 	return (InAssetData.AssetClassPath == FTopLevelAssetPath(TEXT("/Script/CoreUObject"), TEXT("Class"))) // Ignore non-class class items
+PRAGMA_DISABLE_DEPRECATION_WARNINGS
 		&& TryConvertInternalPathToVirtual(InUseFolderPaths ? InAssetData.PackagePath : InAssetData.ObjectPath, OutPath);
+PRAGMA_ENABLE_DEPRECATION_WARNINGS
 }
 
 bool UContentBrowserClassDataSource::IsKnownClassPath(const FName InPackagePath) const
diff --git a/Engine/Plugins/Runtime/AssetTags/Source/AssetTags/Private/AssetTagsSubsystem.cpp b/Engine/Plugins/Runtime/AssetTags/Source/AssetTags/Private/AssetTagsSubsystem.cpp
--- a/Engine/Plugins/Runtime/AssetTags/Source/AssetTags/Private/AssetTagsSubsystem.cpp
+++ b/Engine/Plugins/Runtime/AssetTags/Source/AssetTags/Private/AssetTagsSubsystem.cpp
@@ -168,6 +168,13 @@ bool UAssetTagsSubsystem::EmptyCollection(const FName Name)
 	return true;
 }
 
+bool UAssetTagsSubsystem::K2_AddAssetToCollection(const FName Name, const FSoftObjectPath& AssetPath)
+{
+PRAGMA_DISABLE_DEPRECATION_WARNINGS
+	return AddAssetToCollection(Name, AssetPath.ToFName());
+PRAGMA_ENABLE_DEPRECATION_WARNINGS
+}
+
 bool UAssetTagsSubsystem::AddAssetToCollection(const FName Name, const FName AssetPathName)
 {
 	ICollectionManager& CollectionManager = FCollectionManagerModule::GetModule().Get();
@@ -189,12 +196,23 @@ bool UAssetTagsSubsystem::AddAssetToCollection(const FName Name, const FName Ass
 
 bool UAssetTagsSubsystem::AddAssetDataToCollection(const FName Name, const FAssetData& AssetData)
 {
+PRAGMA_DISABLE_DEPRECATION_WARNINGS
 	return AddAssetToCollection(Name, AssetData.ObjectPath);
+PRAGMA_ENABLE_DEPRECATION_WARNINGS
 }
 
 bool UAssetTagsSubsystem::AddAssetPtrToCollection(const FName Name, const UObject* AssetPtr)
 {
+PRAGMA_DISABLE_DEPRECATION_WARNINGS
 	return AddAssetToCollection(Name, AssetPtr ? FName(*AssetPtr->GetPathName()) : NAME_None);
+PRAGMA_ENABLE_DEPRECATION_WARNINGS
+}
+
+bool UAssetTagsSubsystem::K2_AddAssetsToCollection(const FName Name, const TArray<FSoftObjectPath>& AssetPaths)
+{
+PRAGMA_DISABLE_DEPRECATION_WARNINGS
+	return AddAssetsToCollection(Name, UE::SoftObjectPath::Private::ConvertSoftObjectPaths(AssetPaths));
+PRAGMA_ENABLE_DEPRECATION_WARNINGS
 }
 
 bool UAssetTagsSubsystem::AddAssetsToCollection(const FName Name, const TArray<FName>& AssetPathNames)
@@ -222,9 +240,11 @@ bool UAssetTagsSubsystem::AddAssetDatasToCollection(const FName Name, const TArr
 	AssetPathNames.Reserve(AssetDatas.Num());
 	for (const FAssetData& AssetData : AssetDatas)
 	{
+PRAGMA_DISABLE_DEPRECATION_WARNINGS
 		AssetPathNames.Add(AssetData.ObjectPath);
 	}
 	return AddAssetsToCollection(Name, AssetPathNames);
+PRAGMA_ENABLE_DEPRECATION_WARNINGS
 }
 
 bool UAssetTagsSubsystem::AddAssetPtrsToCollection(const FName Name, const TArray<UObject*>& AssetPtrs)
@@ -235,7 +255,16 @@ bool UAssetTagsSubsystem::AddAssetPtrsToCollection(const FName Name, const TArra
 	{
 		AssetPathNames.Add(AssetPtr ? FName(*AssetPtr->GetPathName()) : NAME_None);
 	}
+PRAGMA_DISABLE_DEPRECATION_WARNINGS
 	return AddAssetsToCollection(Name, AssetPathNames);
+PRAGMA_ENABLE_DEPRECATION_WARNINGS
+}
+
+bool UAssetTagsSubsystem::K2_RemoveAssetFromCollection(const FName Name, const FSoftObjectPath& AssetPath)
+{
+PRAGMA_DISABLE_DEPRECATION_WARNINGS
+	return RemoveAssetFromCollection(Name, AssetPath.ToFName());
+PRAGMA_ENABLE_DEPRECATION_WARNINGS
 }
 
 bool UAssetTagsSubsystem::RemoveAssetFromCollection(const FName Name, const FName AssetPathName)
@@ -259,12 +288,23 @@ bool UAssetTagsSubsystem::RemoveAssetFromCollection(const FName Name, const FNam
 
 bool UAssetTagsSubsystem::RemoveAssetDataFromCollection(const FName Name, const FAssetData& AssetData)
 {
+PRAGMA_DISABLE_DEPRECATION_WARNINGS
 	return RemoveAssetFromCollection(Name, AssetData.ObjectPath);
+PRAGMA_ENABLE_DEPRECATION_WARNINGS
 }
 
 bool UAssetTagsSubsystem::RemoveAssetPtrFromCollection(const FName Name, const UObject* AssetPtr)
 {
+PRAGMA_DISABLE_DEPRECATION_WARNINGS
 	return RemoveAssetFromCollection(Name, AssetPtr ? FName(*AssetPtr->GetPathName()) : NAME_None);
+PRAGMA_ENABLE_DEPRECATION_WARNINGS
+}
+
+bool UAssetTagsSubsystem::K2_RemoveAssetsFromCollection(const FName Name, const TArray<FSoftObjectPath>& AssetPaths)
+{
+PRAGMA_DISABLE_DEPRECATION_WARNINGS
+	return RemoveAssetsFromCollection(Name, UE::SoftObjectPath::Private::ConvertSoftObjectPaths(AssetPaths));
+PRAGMA_ENABLE_DEPRECATION_WARNINGS
 }
 
 bool UAssetTagsSubsystem::RemoveAssetsFromCollection(const FName Name, const TArray<FName>& AssetPathNames)
@@ -292,9 +332,11 @@ bool UAssetTagsSubsystem::RemoveAssetDatasFromCollection(const FName Name, const
 	AssetPathNames.Reserve(AssetDatas.Num());
 	for (const FAssetData& AssetData : AssetDatas)
 	{
+PRAGMA_DISABLE_DEPRECATION_WARNINGS
 		AssetPathNames.Add(AssetData.ObjectPath);
 	}
 	return RemoveAssetsFromCollection(Name, AssetPathNames);
+PRAGMA_ENABLE_DEPRECATION_WARNINGS
 }
 
 bool UAssetTagsSubsystem::RemoveAssetPtrsFromCollection(const FName Name, const TArray<UObject*>& AssetPtrs)
@@ -305,7 +347,9 @@ bool UAssetTagsSubsystem::RemoveAssetPtrsFromCollection(const FName Name, const
 	{
 		AssetPathNames.Add(AssetPtr ? FName(*AssetPtr->GetPathName()) : NAME_None);
 	}
+PRAGMA_DISABLE_DEPRECATION_WARNINGS
 	return RemoveAssetsFromCollection(Name, AssetPathNames);
+PRAGMA_ENABLE_DEPRECATION_WARNINGS
 }
 
 #endif	// WITH_EDITOR
@@ -387,10 +431,13 @@ TArray<FAssetData> UAssetTagsSubsystem::GetAssetsInCollection(const FName Name)
 			TArray<FName> AssetPathNames;
 			CollectionManager.GetAssetsInCollection(ResolvedNameAndType.Name, ResolvedNameAndType.Type, AssetPathNames);
 
-			Assets.Reserve(AssetPathNames.Num());
-			for (const FName& AssetPathName : AssetPathNames)
+PRAGMA_DISABLE_DEPRECATION_WARNINGS
+			TArray<FSoftObjectPath> AssetPaths = UE::SoftObjectPath::Private::ConvertObjectPathNames(AssetPathNames);
+PRAGMA_ENABLE_DEPRECATION_WARNINGS
+			Assets.Reserve(AssetPaths.Num());
+			for (const FSoftObjectPath& AssetPath : AssetPaths)
 			{
-				FAssetData AssetData = AssetRegistry.GetAssetByObjectPath(AssetPathName);
+				FAssetData AssetData = AssetRegistry.GetAssetByObjectPath(AssetPath);
 				if (AssetData.IsValid())
 				{
 					Assets.Add(MoveTemp(AssetData));
@@ -414,6 +461,13 @@ TArray<FAssetData> UAssetTagsSubsystem::GetAssetsInCollection(const FName Name)
 	return Assets;
 }
 
+TArray<FName> UAssetTagsSubsystem::K2_GetCollectionsContainingAsset(const FSoftObjectPath& AssetPath)
+{
+PRAGMA_DISABLE_DEPRECATION_WARNINGS
+	return GetCollectionsContainingAsset(AssetPath.ToFName());
+PRAGMA_ENABLE_DEPRECATION_WARNINGS
+}
+
 TArray<FName> UAssetTagsSubsystem::GetCollectionsContainingAsset(const FName AssetPathName)
 {
 	TArray<FName> CollectionNames;
@@ -464,7 +518,9 @@ TArray<FName> UAssetTagsSubsystem::GetCollectionsContainingAssetData(const FAsse
 	// Note: Use the path name version as the common implementation as:
 	//  1) The path name is always required for the collection manager implementation
 	//  2) The FAssetData for the asset registry implementation *must* come from the asset registry (as the tags are added at cook-time, and missing if FAssetData is generated from a UObject* at runtime)
+PRAGMA_DISABLE_DEPRECATION_WARNINGS
 	return GetCollectionsContainingAsset(AssetData.ObjectPath);
+PRAGMA_ENABLE_DEPRECATION_WARNINGS
 }
 
 TArray<FName> UAssetTagsSubsystem::GetCollectionsContainingAssetPtr(const UObject* AssetPtr)
@@ -472,5 +528,7 @@ TArray<FName> UAssetTagsSubsystem::GetCollectionsContainingAssetPtr(const UObjec
 	// Note: Use the path name version as the common implementation as:
 	//  1) The path name is always required for the collection manager implementation
 	//  2) The FAssetData for the asset registry implementation *must* come from the asset registry (as the tags are added at cook-time, and missing if FAssetData is generated from a UObject* at runtime)
+PRAGMA_DISABLE_DEPRECATION_WARNINGS
 	return GetCollectionsContainingAsset(AssetPtr ? FName(*AssetPtr->GetPathName()) : NAME_None);
+PRAGMA_ENABLE_DEPRECATION_WARNINGS
 }
diff --git a/Engine/Plugins/Runtime/AssetTags/Source/AssetTags/Public/AssetTagsSubsystem.h b/Engine/Plugins/Runtime/AssetTags/Source/AssetTags/Public/AssetTagsSubsystem.h
--- a/Engine/Plugins/Runtime/AssetTags/Source/AssetTags/Public/AssetTagsSubsystem.h
+++ b/Engine/Plugins/Runtime/AssetTags/Source/AssetTags/Public/AssetTagsSubsystem.h
@@ -86,6 +86,10 @@ public:
 	 *
 	 * @return True if the collection was modified, false otherwise (see the output log for details on error).
 	 */
+	UFUNCTION(BlueprintCallable, Category="AssetTags", DisplayName="Add Asset To Collection")
+	bool K2_AddAssetToCollection(const FName Name, const FSoftObjectPath& AssetPath);
+
+	UE_DEPRECATED(5.1, "Names containing full asset paths are deprecated. Use Soft Object Path instead.")
 	UFUNCTION(BlueprintCallable, Category="AssetTags")
 	bool AddAssetToCollection(const FName Name, const FName AssetPathName);
 
@@ -119,6 +123,10 @@ public:
 	 *
 	 * @return True if the collection was modified, false otherwise (see the output log for details on error).
 	 */
+	UFUNCTION(BlueprintCallable, Category="AssetTags", DisplayName="Add Assets To Collection")
+	bool K2_AddAssetsToCollection(const FName Name, const TArray<FSoftObjectPath>& AssetPaths);
+	 
+	UE_DEPRECATED(5.1, "Names containing full asset paths are deprecated. Use Soft Object Path instead.")
 	UFUNCTION(BlueprintCallable, Category="AssetTags")
 	bool AddAssetsToCollection(const FName Name, const TArray<FName>& AssetPathNames);
 
@@ -148,10 +156,14 @@ public:
 	 * Remove the given asset from the given collection.
 	 * 
 	 * @param Name Name of the collection to modify.
-	 * @param AssetPathName Asset to remove (its path name, eg) /Game/MyFolder/MyAsset.MyAsset).
+	 * @param AssetPath Asset to remove (its path, eg) /Game/MyFolder/MyAsset.MyAsset).
 	 *
 	 * @return True if the collection was modified, false otherwise (see the output log for details on error).
 	 */
+	UFUNCTION(BlueprintCallable, Category="AssetTags", DisplayName="Remove Asset From Collection")
+	bool K2_RemoveAssetFromCollection(const FName Name, const FSoftObjectPath& AssetPath);
+
+	UE_DEPRECATED(5.1, "Names containing full asset paths are deprecated, use Soft Object Path instead.")
 	UFUNCTION(BlueprintCallable, Category="AssetTags")
 	bool RemoveAssetFromCollection(const FName Name, const FName AssetPathName);
 
@@ -185,6 +197,10 @@ public:
 	 *
 	 * @return True if the collection was modified, false otherwise (see the output log for details on error).
 	 */
+	UFUNCTION(BlueprintCallable, Category="AssetTags", DisplayName="Remove Assets From Collection")
+	bool K2_RemoveAssetsFromCollection(const FName Name, const TArray<FSoftObjectPath>& AssetPaths);
+
+	UE_DEPRECATED(5.1, "Names containing full asset paths are deprecated. Use Soft Object Path instead.")
 	UFUNCTION(BlueprintCallable, Category="AssetTags")
 	bool RemoveAssetsFromCollection(const FName Name, const TArray<FName>& AssetPathNames);
 
@@ -246,6 +262,10 @@ public:
 	 *
 	 * @return Names of the collections that contain the asset.
 	 */
+	UFUNCTION(BlueprintCallable, Category = "AssetTags", DisplayName="Get Collections Containing Asset")
+	TArray<FName> K2_GetCollectionsContainingAsset(const FSoftObjectPath& AssetPath);
+
+	UE_DEPRECATED(5.1, "Names containing full asset paths are deprecated. Use Soft Object Path instead.")
 	UFUNCTION(BlueprintCallable, Category="AssetTags")
 	TArray<FName> GetCollectionsContainingAsset(const FName AssetPathName);
 
diff --git a/Engine/Source/Developer/AssetTools/Private/CollectionAssetManagement.cpp b/Engine/Source/Developer/AssetTools/Private/CollectionAssetManagement.cpp
--- a/Engine/Source/Developer/AssetTools/Private/CollectionAssetManagement.cpp
+++ b/Engine/Source/Developer/AssetTools/Private/CollectionAssetManagement.cpp
@@ -42,13 +42,13 @@ void FCollectionAssetManagement::SetCurrentAssets(const TArray<FAssetData>& Curr
 	CurrentAssetPaths.Empty();
 	for (const FAssetData& AssetData : CurrentAssets)
 	{
-		CurrentAssetPaths.Add(AssetData.ObjectPath);
+		CurrentAssetPaths.Add(AssetData.GetSoftObjectPath());
 	}
 
 	UpdateAssetManagementState();
 }
 
-void FCollectionAssetManagement::SetCurrentAssetPaths(const TArray<FName>& CurrentAssets)
+void FCollectionAssetManagement::SetCurrentAssetPaths(const TArray<FSoftObjectPath>& CurrentAssets)
 {
 	CurrentAssetPaths.Empty();
 	CurrentAssetPaths.Append(CurrentAssets);
@@ -60,7 +60,9 @@ void FCollectionAssetManagement::AddCurrentAssetsToCollection(FCollectionNameTyp
 {
 	FCollectionManagerModule& CollectionManagerModule = FCollectionManagerModule::GetModule();
 
-	const TArray<FName> ObjectPaths = CurrentAssetPaths.Array();
+PRAGMA_DISABLE_DEPRECATION_WARNINGS
+	const TArray<FName> ObjectPaths = UE::SoftObjectPath::Private::ConvertSoftObjectPaths(CurrentAssetPaths.Array());
+PRAGMA_ENABLE_DEPRECATION_WARNINGS
 
 	FText ResultText;
 	bool bSuccess = false;
@@ -99,7 +101,9 @@ void FCollectionAssetManagement::RemoveCurrentAssetsFromCollection(FCollectionNa
 {
 	FCollectionManagerModule& CollectionManagerModule = FCollectionManagerModule::GetModule();
 
-	const TArray<FName> ObjectPaths = CurrentAssetPaths.Array();
+PRAGMA_DISABLE_DEPRECATION_WARNINGS
+	const TArray<FName> ObjectPaths = UE::SoftObjectPath::Private::ConvertSoftObjectPaths(CurrentAssetPaths.Array());
+PRAGMA_ENABLE_DEPRECATION_WARNINGS
 
 	FText ResultText;
 	bool bSuccess = false;
@@ -169,7 +173,9 @@ void FCollectionAssetManagement::UpdateAssetManagementState()
 	if (CurrentAssetPaths.Num() == 1)
 	{
 		TArray<FCollectionNameType> MatchedCollections;
-		CollectionManagerModule.Get().GetCollectionsContainingObject(*CurrentAssetPaths.CreateConstIterator(), MatchedCollections, ECollectionRecursionFlags::Self);
+PRAGMA_DISABLE_DEPRECATION_WARNINGS
+		CollectionManagerModule.Get().GetCollectionsContainingObject(CurrentAssetPaths.CreateConstIterator()->ToFName(), MatchedCollections, ECollectionRecursionFlags::Self);
+PRAGMA_ENABLE_DEPRECATION_WARNINGS
 
 		for (const FCollectionNameType& CollectionKey : MatchedCollections)
 		{
@@ -177,7 +183,9 @@ void FCollectionAssetManagement::UpdateAssetManagementState()
 		}
 
 		MatchedCollections.Reset();
-		CollectionManagerModule.Get().GetCollectionsContainingObject(*CurrentAssetPaths.CreateConstIterator(), MatchedCollections, ECollectionRecursionFlags::Children);
+PRAGMA_DISABLE_DEPRECATION_WARNINGS
+		CollectionManagerModule.Get().GetCollectionsContainingObject(CurrentAssetPaths.CreateConstIterator()->ToFName(), MatchedCollections, ECollectionRecursionFlags::Children);
+PRAGMA_ENABLE_DEPRECATION_WARNINGS
 
 		for (const FCollectionNameType& CollectionKey : MatchedCollections)
 		{
@@ -189,7 +197,9 @@ void FCollectionAssetManagement::UpdateAssetManagementState()
 	}
 	else
 	{
-		const TArray<FName> ObjectPaths = CurrentAssetPaths.Array();
+PRAGMA_DISABLE_DEPRECATION_WARNINGS
+		const TArray<FName> ObjectPaths = UE::SoftObjectPath::Private::ConvertSoftObjectPaths(CurrentAssetPaths.Array());
+PRAGMA_ENABLE_DEPRECATION_WARNINGS
 
 		TMap<FCollectionNameType, TArray<FName>> CollectionsAndMatchedObjects;
 		CollectionManagerModule.Get().GetCollectionsContainingObjects(ObjectPaths, CollectionsAndMatchedObjects);
@@ -234,7 +244,9 @@ void FCollectionAssetManagement::HandleAssetsAddedToCollection(const FCollection
 	bool bNeedsUpdate = false;
 	for (const FName& AssetPath : AssetsAdded)
 	{
-		if (CurrentAssetPaths.Contains(AssetPath))
+PRAGMA_DISABLE_DEPRECATION_WARNINGS
+		if (CurrentAssetPaths.Contains(FSoftObjectPath(AssetPath)))
+PRAGMA_ENABLE_DEPRECATION_WARNINGS
 		{
 			bNeedsUpdate = true;
 			break;
@@ -253,7 +265,9 @@ void FCollectionAssetManagement::HandleAssetsRemovedFromCollection(const FCollec
 	bool bNeedsUpdate = false;
 	for (const FName& AssetPath : AssetsRemoved)
 	{
-		if (CurrentAssetPaths.Contains(AssetPath))
+PRAGMA_DISABLE_DEPRECATION_WARNINGS
+		if (CurrentAssetPaths.Contains(FSoftObjectPath(AssetPath)))
+PRAGMA_ENABLE_DEPRECATION_WARNINGS
 		{
 			bNeedsUpdate = true;
 			break;
diff --git a/Engine/Source/Developer/AssetTools/Public/CollectionAssetManagement.h b/Engine/Source/Developer/AssetTools/Public/CollectionAssetManagement.h
--- a/Engine/Source/Developer/AssetTools/Public/CollectionAssetManagement.h
+++ b/Engine/Source/Developer/AssetTools/Public/CollectionAssetManagement.h
@@ -21,7 +21,7 @@ public:
 	void SetCurrentAssets(const TArray<FAssetData>& CurrentAssets);
 
 	/** Set the asset paths that we are currently observing and managing the collection state of */
-	void SetCurrentAssetPaths(const TArray<FName>& CurrentAssets);
+	void SetCurrentAssetPaths(const TArray<FSoftObjectPath>& CurrentAssets);
 
 	/** Add the current assets to the given collection */
 	void AddCurrentAssetsToCollection(FCollectionNameType InCollectionKey);
@@ -55,7 +55,7 @@ private:
 	void HandleAssetsRemovedFromCollection(const FCollectionNameType& Collection, const TArray<FName>& AssetsRemoved);
 
 	/** Set of asset paths that we are currently observing and managing the collection state of */
-	TSet<FName> CurrentAssetPaths;
+	TSet<FSoftObjectPath> CurrentAssetPaths;
 
 	/** Mapping between a collection and its asset management state (based on the current assets). A missing item is assumed to be unchecked */
 	TMap<FCollectionNameType, ECheckBoxState> AssetManagementState;
diff --git a/Engine/Source/Editor/ContentBrowser/Private/AssetContextMenu.cpp b/Engine/Source/Editor/ContentBrowser/Private/AssetContextMenu.cpp
--- a/Engine/Source/Editor/ContentBrowser/Private/AssetContextMenu.cpp
+++ b/Engine/Source/Editor/ContentBrowser/Private/AssetContextMenu.cpp
@@ -174,7 +174,7 @@ TSharedRef<SWidget> FAssetContextMenu::MakeContextMenu(TArrayView<const FContent
 			FAssetData ItemAssetData;
 			if (SelectedItem.Legacy_TryGetAssetData(ItemAssetData))
 			{
-				ObjectPaths.Add(ItemAssetData.ObjectPath.ToString());
+				ObjectPaths.Add(ItemAssetData.GetObjectPathString());
 			}
 		}
 
@@ -759,7 +759,9 @@ bool FAssetContextMenu::AddCollectionMenuOptions(UToolMenu* Menu)
 				SelectedItemCollectionIds.Add(ItemCollectionId);
 			}
 		}
-		QuickAssetManagement->SetCurrentAssetPaths(SelectedItemCollectionIds);
+PRAGMA_DISABLE_DEPRECATION_WARNINGS
+		QuickAssetManagement->SetCurrentAssetPaths(UE::SoftObjectPath::Private::ConvertObjectPathNames(SelectedItemCollectionIds));
+PRAGMA_ENABLE_DEPRECATION_WARNINGS
 
 		Section.AddSubMenu(
 			"ManageCollections",
diff --git a/Engine/Source/Editor/ContentBrowser/Private/AssetViewWidgets.cpp b/Engine/Source/Editor/ContentBrowser/Private/AssetViewWidgets.cpp
--- a/Engine/Source/Editor/ContentBrowser/Private/AssetViewWidgets.cpp
+++ b/Engine/Source/Editor/ContentBrowser/Private/AssetViewWidgets.cpp
@@ -968,7 +968,9 @@ TSharedRef<SWidget> SAssetViewItem::CreateToolTipWidget() const
 				ICollectionManager& CollectionManager = FCollectionManagerModule::GetModule().Get();
 
 				TArray<FCollectionNameType> CollectionsContainingObject;
+PRAGMA_DISABLE_DEPRECATION_WARNINGS
 				CollectionManager.GetCollectionsContainingObject(ItemAssetData.ObjectPath, CollectionsContainingObject);
+PRAGMA_ENABLE_DEPRECATION_WARNINGS
 
 				if (CollectionsContainingObject.Num() > 0)
 				{
diff --git a/Engine/Source/Editor/ContentBrowser/Private/CollectionAssetRegistryBridge.cpp b/Engine/Source/Editor/ContentBrowser/Private/CollectionAssetRegistryBridge.cpp
--- a/Engine/Source/Editor/ContentBrowser/Private/CollectionAssetRegistryBridge.cpp
+++ b/Engine/Source/Editor/ContentBrowser/Private/CollectionAssetRegistryBridge.cpp
@@ -91,17 +91,19 @@ public:
 		else
 		{
 			// Keep track of visted redirectors in case we loop.
-			TSet<FName> VisitedRedirectors;
+			TSet<FSoftObjectPath> VisitedRedirectors;
 
 			// Use the asset registry to avoid loading the object
+PRAGMA_DISABLE_DEPRECATION_WARNINGS
 			FAssetData ObjectAssetData = AssetRegistryModule.Get().GetAssetByObjectPath(InObjectPath, true);
+PRAGMA_ENABLE_DEPRECATION_WARNINGS
 			while (ObjectAssetData.IsValid() && ObjectAssetData.IsRedirector())
 			{
 				// Check to see if we've already seen this path before, it's possible we might have found a redirector loop.
-				if ( VisitedRedirectors.Contains(ObjectAssetData.ObjectPath) )
+				if ( VisitedRedirectors.Contains(ObjectAssetData.ToSoftObjectPath()) )
 				{
 					UE_LOG(LogContentBrowser, Error, TEXT("Redirector Loop Found!"));
-					for ( FName Redirector : VisitedRedirectors )
+					for ( const FSoftObjectPath& Redirector : VisitedRedirectors )
 					{
 						UE_LOG(LogContentBrowser, Error, TEXT("Redirector: %s"), *Redirector.ToString());
 					}
@@ -110,7 +112,7 @@ public:
 					break;
 				}
 
-				VisitedRedirectors.Add(ObjectAssetData.ObjectPath);
+				VisitedRedirectors.Add(ObjectAssetData.ToSoftObjectPath());
 
 				// Get the destination object from the meta-data rather than load the redirector object, as 
 				// loading a redirector will also load the object it points to, which could cause a large hitch
@@ -118,7 +120,7 @@ public:
 				if (ObjectAssetData.GetTagValue("DestinationObject", DestinationObjectPath))
 				{
 					ConstructorHelpers::StripObjectClass(DestinationObjectPath);
-					ObjectAssetData = AssetRegistryModule.Get().GetAssetByObjectPath(*DestinationObjectPath);
+					ObjectAssetData = AssetRegistryModule.Get().GetAssetByObjectPath(FSoftObjectPath(DestinationObjectPath));
 				}
 				else
 				{
@@ -126,7 +128,9 @@ public:
 				}
 			}
 
+PRAGMA_DISABLE_DEPRECATION_WARNINGS
 			OutNewObjectPath = ObjectAssetData.ObjectPath;
+PRAGMA_ENABLE_DEPRECATION_WARNINGS
 		}
 
 		return OutNewObjectPath != NAME_None && InObjectPath != OutNewObjectPath;
@@ -179,7 +183,9 @@ void FCollectionAssetRegistryBridge::OnAssetRenamed(const FAssetData& AssetData,
 	FCollectionManagerModule& CollectionManagerModule = FCollectionManagerModule::GetModule();
 
 	// Notify the collections manager that an asset has been renamed
+PRAGMA_DISABLE_DEPRECATION_WARNINGS
 	CollectionManagerModule.Get().HandleObjectRenamed(*OldObjectPath, AssetData.ObjectPath);
+PRAGMA_ENABLE_DEPRECATION_WARNINGS
 }
 
 void FCollectionAssetRegistryBridge::OnAssetRemoved(const FAssetData& AssetData)
@@ -190,12 +196,16 @@ void FCollectionAssetRegistryBridge::OnAssetRemoved(const FAssetData& AssetData)
 	{
 		// Notify the collections manager that a redirector has been removed
 		// This will attempt to re-save any collections that still have a reference to this redirector in their on-disk collection data
+PRAGMA_DISABLE_DEPRECATION_WARNINGS
 		CollectionManagerModule.Get().HandleRedirectorDeleted(AssetData.ObjectPath);
+PRAGMA_ENABLE_DEPRECATION_WARNINGS
 	}
 	else
 	{
 		// Notify the collections manager that an asset has been removed
+PRAGMA_DISABLE_DEPRECATION_WARNINGS
 		CollectionManagerModule.Get().HandleObjectDeleted(AssetData.ObjectPath);
+PRAGMA_ENABLE_DEPRECATION_WARNINGS
 	}
 }
 
diff --git a/Engine/Source/Editor/ContentBrowser/Private/ContentBrowserUtils.cpp b/Engine/Source/Editor/ContentBrowser/Private/ContentBrowserUtils.cpp
--- a/Engine/Source/Editor/ContentBrowser/Private/ContentBrowserUtils.cpp
+++ b/Engine/Source/Editor/ContentBrowser/Private/ContentBrowserUtils.cpp
@@ -56,6 +56,7 @@
 #include "Widgets/SBoxPanel.h"
 #include "Widgets/SCompoundWidget.h"
 #include "Widgets/Text/STextBlock.h"
+#include "Algo/Transform.h"
 
 class SWidget;
 struct FGeometry;
@@ -577,10 +578,13 @@ void ContentBrowserUtils::ConvertLegacySelectionToVirtualPaths(TArrayView<const
 
 void ContentBrowserUtils::AppendAssetFilterToContentBrowserFilter(const FARFilter& InAssetFilter, const TSharedPtr<FPathPermissionList>& InAssetClassPermissionList, const TSharedPtr<FPathPermissionList>& InFolderPermissionList, FContentBrowserDataFilter& OutDataFilter)
 {
-	if (InAssetFilter.ObjectPaths.Num() > 0 || InAssetFilter.TagsAndValues.Num() > 0 || InAssetFilter.bIncludeOnlyOnDiskAssets)
+	if (InAssetFilter.SoftObjectPaths.Num() > 0 || InAssetFilter.TagsAndValues.Num() > 0 || InAssetFilter.bIncludeOnlyOnDiskAssets)
 	{
 		FContentBrowserDataObjectFilter& ObjectFilter = OutDataFilter.ExtraFilters.FindOrAddFilter<FContentBrowserDataObjectFilter>();
-		ObjectFilter.ObjectNamesToInclude = InAssetFilter.ObjectPaths;
+PRAGMA_DISABLE_DEPRECATION_WARNINGS
+		// TODO: Modify this API to also use FSoftObjectPath with deprecation
+		ObjectFilter.ObjectNamesToInclude = UE::SoftObjectPath::Private::ConvertSoftObjectPaths(InAssetFilter.SoftObjectPaths);
+PRAGMA_DISABLE_DEPRECATION_WARNINGS
 		ObjectFilter.TagsAndValuesToInclude = InAssetFilter.TagsAndValues;
 		ObjectFilter.bOnDiskObjectsOnly = InAssetFilter.bIncludeOnlyOnDiskAssets;
 	}
diff --git a/Engine/Source/Editor/ContentBrowser/Private/SAssetDialog.cpp b/Engine/Source/Editor/ContentBrowser/Private/SAssetDialog.cpp
--- a/Engine/Source/Editor/ContentBrowser/Private/SAssetDialog.cpp
+++ b/Engine/Source/Editor/ContentBrowser/Private/SAssetDialog.cpp
@@ -993,7 +993,7 @@ void SAssetDialog::UpdateInputValidity()
 			else if(bAllowExistingAsset && AssetClassNames.Num() > 1) // If for some reason we have multiple names, perform additional logic here...
 			{
 				FAssetRegistryModule& AssetRegistryModule = FModuleManager::LoadModuleChecked<FAssetRegistryModule>("AssetRegistry");
-				FAssetData ExistingAsset = AssetRegistryModule.Get().GetAssetByObjectPath(FName(*ObjectPath));
+				FAssetData ExistingAsset = AssetRegistryModule.Get().GetAssetByObjectPath(FSoftObjectPath(ObjectPath));
 				if (ExistingAsset.IsValid() && !AssetClassNames.Contains(ExistingAsset.AssetClassPath))
 				{
 					const FString ObjectName = FPackageName::ObjectPathToObjectName(ObjectPath);
@@ -1065,7 +1065,7 @@ void SAssetDialog::CommitObjectPathForSave()
 			if ( ExistingAssetPolicy == ESaveAssetDialogExistingAssetPolicy::AllowButWarn )
 			{
 				FAssetRegistryModule& AssetRegistryModule = FModuleManager::LoadModuleChecked<FAssetRegistryModule>("AssetRegistry");
-				FAssetData ExistingAsset = AssetRegistryModule.Get().GetAssetByObjectPath(FName(*ObjectPath));
+				FAssetData ExistingAsset = AssetRegistryModule.Get().GetAssetByObjectPath(FSoftObjectPath(ObjectPath));
 				if ( ExistingAsset.IsValid() && AssetClassNames.Contains(ExistingAsset.AssetClassPath) )
 				{
 					EAppReturnType::Type ShouldReplace = FMessageDialog::Open( EAppMsgType::YesNo, FText::Format(LOCTEXT("ReplaceAssetMessage", "{0} already exists. Do you want to replace it?"), FText::FromString(CurrentlyEnteredAssetName)) );
diff --git a/Engine/Source/Editor/ContentBrowser/Private/SAssetView.cpp b/Engine/Source/Editor/ContentBrowser/Private/SAssetView.cpp
--- a/Engine/Source/Editor/ContentBrowser/Private/SAssetView.cpp
+++ b/Engine/Source/Editor/ContentBrowser/Private/SAssetView.cpp
@@ -3651,7 +3651,7 @@ bool SAssetView::CanOpenContextMenu() const
 			FAssetData ItemAssetData;
 			if (Item->GetItem().Legacy_TryGetAssetData(ItemAssetData))
 			{
-				ObjectPaths.Add(ItemAssetData.ObjectPath.ToString());
+				ObjectPaths.Add(ItemAssetData.GetObjectPathString());
 			}
 		}
 
diff --git a/Engine/Source/Editor/ContentBrowser/Private/SCollectionView.cpp b/Engine/Source/Editor/ContentBrowser/Private/SCollectionView.cpp
--- a/Engine/Source/Editor/ContentBrowser/Private/SCollectionView.cpp
+++ b/Engine/Source/Editor/ContentBrowser/Private/SCollectionView.cpp
@@ -629,7 +629,7 @@ TArray<FCollectionNameType> SCollectionView::GetSelectedCollections() const
 	return RetArray;
 }
 
-void SCollectionView::SetSelectedAssetPaths(const TArray<FName>& SelectedAssets)
+void SCollectionView::SetSelectedAssetPaths(const TArray<FSoftObjectPath>& SelectedAssets)
 {
 	if ( QuickAssetManagement.IsValid() )
 	{
@@ -1324,16 +1324,18 @@ FReply SCollectionView::HandleDragDropOnCollectionItem(TSharedRef<FCollectionIte
 		TSharedPtr<FAssetDragDropOp> DragDropOp = StaticCastSharedPtr<FAssetDragDropOp>(Operation);
 		const TArray<FAssetData>& DroppedAssets = DragDropOp->GetAssets();
 
-		TArray<FName> ObjectPaths;
+		TArray<FSoftObjectPath> ObjectPaths;
 		ObjectPaths.Reserve(DroppedAssets.Num());
 		for (const FAssetData& AssetData : DroppedAssets)
 		{
-			ObjectPaths.Add(AssetData.ObjectPath);
+			ObjectPaths.Add(AssetData.GetSoftObjectPath());
 		}
 
 		int32 NumAdded = 0;
 		FText Message;
-		if (CollectionManagerModule.Get().AddToCollection(CollectionItem->CollectionName, CollectionItem->CollectionType, ObjectPaths, &NumAdded))
+PRAGMA_DISABLE_DEPRECATION_WARNINGS
+		if (CollectionManagerModule.Get().AddToCollection(CollectionItem->CollectionName, CollectionItem->CollectionType, UE::SoftObjectPath::Private::ConvertSoftObjectPaths(ObjectPaths), &NumAdded))
+PRAGMA_ENABLE_DEPRECATION_WARNINGS
 		{
 			if (DroppedAssets.Num() == 1)
 			{
diff --git a/Engine/Source/Editor/ContentBrowser/Private/SCollectionView.h b/Engine/Source/Editor/ContentBrowser/Private/SCollectionView.h
--- a/Engine/Source/Editor/ContentBrowser/Private/SCollectionView.h
+++ b/Engine/Source/Editor/ContentBrowser/Private/SCollectionView.h
@@ -105,7 +105,7 @@ public:
 	TArray<FCollectionNameType> GetSelectedCollections() const;
 
 	/** Let the collections view know that the list of selected assets has changed, so that it can update the quick asset management check boxes */
-	void SetSelectedAssetPaths(const TArray<FName>& SelectedAssets);
+	void SetSelectedAssetPaths(const TArray<FSoftObjectPath>& SelectedAssets);
 
 	/** Sets the state of the collection view to the one described by the history data */
 	void ApplyHistoryData ( const FHistoryData& History );
diff --git a/Engine/Source/Editor/ContentBrowser/Private/SContentBrowser.cpp b/Engine/Source/Editor/ContentBrowser/Private/SContentBrowser.cpp
--- a/Engine/Source/Editor/ContentBrowser/Private/SContentBrowser.cpp
+++ b/Engine/Source/Editor/ContentBrowser/Private/SContentBrowser.cpp
@@ -1568,13 +1568,13 @@ void SContentBrowser::SyncToLegacy(TArrayView<const FAssetData> AssetDataList, T
 	TArray<FAssetData> NewItemsToSync;
 	for (int32 i = 0; i < AssetDataList.Num(); ++i)
 	{
-		if (FolderPermissions->PassesStartsWithFilter(AssetDataList[i].ObjectPath))
+		if (FolderPermissions->PassesStartsWithFilter(AssetDataList[i].GetObjectPathString()))
 		{
 			NewItemsToSync.Add(AssetDataList[i]);
 		}
 		else
 		{
-			TArray<FContentBrowserItemPath> Aliases = ContentBrowserData->GetAliasesForPath(AssetDataList[i].ObjectPath);
+			TArray<FContentBrowserItemPath> Aliases = ContentBrowserData->GetAliasesForPath(*AssetDataList[i].GetObjectPathString());
 			for (const FContentBrowserItemPath& Alias : Aliases)
 			{
 				const FName InternalPath = Alias.GetInternalPathName();
@@ -1656,7 +1656,7 @@ void SContentBrowser::LoadSelectedObjectsIfNeeded()
 	for ( auto AssetIt = SelectedAssets.CreateConstIterator(); AssetIt; ++AssetIt )
 	{
 		const FAssetData& AssetData = *AssetIt;
-		const bool bShowProgressDialog = (!AssetData.IsAssetLoaded() && FEditorFileUtils::IsMapPackageAsset(AssetData.ObjectPath.ToString()));
+		const bool bShowProgressDialog = (!AssetData.IsAssetLoaded() && FEditorFileUtils::IsMapPackageAsset(AssetData.GetObjectPathString()));
 		GWarn->BeginSlowTask(LOCTEXT("LoadingObjects", "Loading Objects..."), bShowProgressDialog);
 
 		(*AssetIt).GetAsset();
@@ -2905,7 +2905,9 @@ void SContentBrowser::OnItemSelectionChanged(const FContentBrowserItem& Selected
 				}
 			}
 
-			CollectionViewPtr->SetSelectedAssetPaths(SelectedCollectionItems);
+PRAGMA_DISABLE_DEPRECATION_WARNINGS
+			CollectionViewPtr->SetSelectedAssetPaths(UE::SoftObjectPath::Private::ConvertObjectPathNames(SelectedCollectionItems));
+PRAGMA_ENABLE_DEPRECATION_WARNINGS
 		}
 
 		if (AssetSelectionChangedDelegate.IsBound())
diff --git a/Engine/Source/Editor/ContentBrowserData/Private/ContentBrowserDataSource.cpp b/Engine/Source/Editor/ContentBrowserData/Private/ContentBrowserDataSource.cpp
--- a/Engine/Source/Editor/ContentBrowserData/Private/ContentBrowserDataSource.cpp
+++ b/Engine/Source/Editor/ContentBrowserData/Private/ContentBrowserDataSource.cpp
@@ -223,11 +223,16 @@ bool UContentBrowserDataSource::EnumerateItemsForObjects(const TArrayView<UObjec
 	return true;
 }
 
-TArray<FContentBrowserItemPath> UContentBrowserDataSource::GetAliasesForPath(const FName InInternalPath) const
+TArray<FContentBrowserItemPath> UContentBrowserDataSource::GetAliasesForPath(const FSoftObjectPath& InInternalPath) const
 {
 	return TArray<FContentBrowserItemPath>();
 }
 
+TArray<FContentBrowserItemPath> UContentBrowserDataSource::GetAliasesForPath(FName InInternalPath) const
+{
+	return GetAliasesForPath(FSoftObjectPath(InInternalPath));
+}
+
 bool UContentBrowserDataSource::IsDiscoveringItems(FText* OutStatus)
 {
 	return false;
diff --git a/Engine/Source/Editor/ContentBrowserData/Private/ContentBrowserDataSubsystem.cpp b/Engine/Source/Editor/ContentBrowserData/Private/ContentBrowserDataSubsystem.cpp
--- a/Engine/Source/Editor/ContentBrowserData/Private/ContentBrowserDataSubsystem.cpp
+++ b/Engine/Source/Editor/ContentBrowserData/Private/ContentBrowserDataSubsystem.cpp
@@ -468,6 +468,20 @@ TArray<FContentBrowserItemPath> UContentBrowserDataSubsystem::GetAliasesForPath(
 	return GetAliasesForPath(InPath.GetInternalPathName());
 }
 
+TArray<FContentBrowserItemPath> UContentBrowserDataSubsystem::GetAliasesForPath(const FSoftObjectPath& InInternalPath) const
+{
+	TArray<FContentBrowserItemPath> Aliases;
+
+	for (const auto& ActiveDataSourcePair : ActiveDataSources)
+	{
+		UContentBrowserDataSource* DataSource = ActiveDataSourcePair.Value;
+		Aliases.Append(DataSource->GetAliasesForPath(InInternalPath));
+	}
+
+	return Aliases;
+}
+
+PRAGMA_DISABLE_DEPRECATION_WARNINGS
 TArray<FContentBrowserItemPath> UContentBrowserDataSubsystem::GetAliasesForPath(const FName InInternalPath) const
 {
 	TArray<FContentBrowserItemPath> Aliases;
@@ -480,6 +494,7 @@ TArray<FContentBrowserItemPath> UContentBrowserDataSubsystem::GetAliasesForPath(
 
 	return Aliases;
 }
+PRAGMA_ENABLE_DEPRECATION_WARNINGS
 
 bool UContentBrowserDataSubsystem::IsDiscoveringItems(TArray<FText>* OutStatus) const
 {
diff --git a/Engine/Source/Editor/ContentBrowserData/Public/ContentBrowserDataSource.h b/Engine/Source/Editor/ContentBrowserData/Public/ContentBrowserDataSource.h
--- a/Engine/Source/Editor/ContentBrowserData/Public/ContentBrowserDataSource.h
+++ b/Engine/Source/Editor/ContentBrowserData/Public/ContentBrowserDataSource.h
@@ -230,7 +230,10 @@ public:
 	 * @param The internal path (or object path) of an asset to get aliases for
 	 * @return All alternative paths that represent the input path (not including the input path itself)
 	 */
-	virtual TArray<FContentBrowserItemPath> GetAliasesForPath(const FName InInternalPath) const;
+	virtual TArray<FContentBrowserItemPath> GetAliasesForPath(const FSoftObjectPath& InInternalPath) const;
+
+	UE_DEPRECATED(5.1, "FNames containing full asset paths are deprecated. Use FSoftObjectPath instead.")
+	TArray<FContentBrowserItemPath> GetAliasesForPath(FName InInternalPath) const;
 
 	/**
 	 * Query whether this data source instance is currently discovering content, and retrieve an optional status message that can be shown in the UI.
diff --git a/Engine/Source/Editor/ContentBrowserData/Public/ContentBrowserDataSubsystem.h b/Engine/Source/Editor/ContentBrowserData/Public/ContentBrowserDataSubsystem.h
--- a/Engine/Source/Editor/ContentBrowserData/Public/ContentBrowserDataSubsystem.h
+++ b/Engine/Source/Editor/ContentBrowserData/Public/ContentBrowserDataSubsystem.h
@@ -219,8 +219,9 @@ public:
 	 * @param The internal path (or object path) of an asset to get aliases for
 	 * @return All alternative paths that represent the input path (not including the input path itself)
 	 */
-	TArray<FContentBrowserItemPath> GetAliasesForPath(const FName InInternalPath) const;
+	TArray<FContentBrowserItemPath> GetAliasesForPath(const FSoftObjectPath& InInternalPath) const;
 	TArray<FContentBrowserItemPath> GetAliasesForPath(const FContentBrowserItemPath InPath) const;
+	TArray<FContentBrowserItemPath> GetAliasesForPath(const FName InInternalPath) const;
 
 	/**
 	 * Query whether any data sources are currently discovering content, and retrieve optional status messages that can be shown in the UI.
diff --git a/Engine/Source/Editor/UnrealEd/Private/DragAndDrop/CollectionDragDropOp.cpp b/Engine/Source/Editor/UnrealEd/Private/DragAndDrop/CollectionDragDropOp.cpp
--- a/Engine/Source/Editor/UnrealEd/Private/DragAndDrop/CollectionDragDropOp.cpp
+++ b/Engine/Source/Editor/UnrealEd/Private/DragAndDrop/CollectionDragDropOp.cpp
@@ -26,17 +26,19 @@ TArray<FAssetData> FCollectionDragDropOp::GetAssets() const
 	FAssetRegistryModule& AssetRegistryModule = FModuleManager::LoadModuleChecked<FAssetRegistryModule>(TEXT("AssetRegistry"));
 	IAssetRegistry& AssetRegistry = AssetRegistryModule.Get();
 
-	TArray<FName> AssetPaths;
+	TArray<FSoftObjectPath> AssetPaths;
 	for (const FCollectionNameType& CollectionNameType : Collections)
 	{
 		TArray<FName> CollectionAssetPaths;
+PRAGMA_DISABLE_DEPRECATION_WARNINGS
 		CollectionManager.GetAssetsInCollection(CollectionNameType.Name, CollectionNameType.Type, CollectionAssetPaths);
-		AssetPaths.Append(CollectionAssetPaths);
+		AssetPaths.Append(UE::SoftObjectPath::Private::ConvertObjectPathNames(CollectionAssetPaths));
+PRAGMA_ENABLE_DEPRECATION_WARNINGS
 	}
 
 	TArray<FAssetData> AssetDatas;
 	AssetDatas.Reserve(AssetPaths.Num());
-	for (const FName& AssetPath : AssetPaths)
+	for (const FSoftObjectPath& AssetPath : AssetPaths)
 	{
 		FAssetData AssetData = AssetRegistry.GetAssetByObjectPath(AssetPath);
 		if (AssetData.IsValid())
