Remove ability to add aliases to collections

diff --git a/Engine/Plugins/Editor/ContentBrowser/ContentBrowserAliasDataSoure/Source/ContentBrowserAliasDataSource/Private/ContentBrowserAliasDataSource.cpp b/Engine/Plugins/Editor/ContentBrowser/ContentBrowserAliasDataSoure/Source/ContentBrowserAliasDataSource/Private/ContentBrowserAliasDataSource.cpp
--- a/Engine/Plugins/Editor/ContentBrowser/ContentBrowserAliasDataSoure/Source/ContentBrowserAliasDataSource/Private/ContentBrowserAliasDataSource.cpp
+++ b/Engine/Plugins/Editor/ContentBrowser/ContentBrowserAliasDataSoure/Source/ContentBrowserAliasDataSource/Private/ContentBrowserAliasDataSource.cpp
@@ -734,14 +734,6 @@ bool UContentBrowserAliasDataSource::UpdateThumbnail(const FContentBrowserItemDa
 
 bool UContentBrowserAliasDataSource::TryGetCollectionId(const FContentBrowserItemData& InItem, FSoftObjectPath& OutCollectionId)
 {
-	if (TSharedPtr<const FContentBrowserAliasItemDataPayload> AliasPayload = StaticCastSharedPtr<const FContentBrowserAliasItemDataPayload>(InItem.GetPayload()))
-	{
-		if (const FAliasData* AliasData = AllAliases.Find(AliasPayload->Alias))
-		{
-			OutCollectionId = FSoftObjectPath(AliasData->ObjectPath);
-			return true;
-		}
-	}
 	return false;
 }
 
