Add ability to filter CB folders for Editability + Fix AliasDataType usage of Public/Private assets

diff --git a/Engine/Plugins/Editor/ContentBrowser/ContentBrowserAliasDataSoure/Source/ContentBrowserAliasDataSource/Private/ContentBrowserAliasDataSource.cpp b/Engine/Plugins/Editor/ContentBrowser/ContentBrowserAliasDataSoure/Source/ContentBrowserAliasDataSource/Private/ContentBrowserAliasDataSource.cpp
--- a/Engine/Plugins/Editor/ContentBrowser/ContentBrowserAliasDataSoure/Source/ContentBrowserAliasDataSource/Private/ContentBrowserAliasDataSource.cpp
+++ b/Engine/Plugins/Editor/ContentBrowser/ContentBrowserAliasDataSoure/Source/ContentBrowserAliasDataSource/Private/ContentBrowserAliasDataSource.cpp
@@ -13,6 +13,7 @@
 #include "Modules/ModuleManager.h"
 #include "CollectionManagerModule.h"
 #include "ICollectionManager.h"
+#include "IContentBrowserDataModule.h"
 
 #include UE_INLINE_GENERATED_CPP_BY_NAME(ContentBrowserAliasDataSource)
 
@@ -670,7 +671,7 @@ bool UContentBrowserAliasDataSource::CanEditItem(const FContentBrowserItemData&
 {
 	if (TSharedPtr<const FContentBrowserAliasItemDataPayload> AliasPayload = StaticCastSharedPtr<const FContentBrowserAliasItemDataPayload>(InItem.GetPayload()))
 	{
-		// Both the alias path and asset path must pass the writable folder filter in order to be editable
+		// Both the alias path and asset path must pass the writable folder filter and editable folder filter in order to be editable
 		const TSharedRef<FPathPermissionList>& WritableFolderFilter = AssetTools->GetWritableFolderPermissionList();
 		if (!WritableFolderFilter->PassesStartsWithFilter(AliasPayload->Alias.Value))
 		{
@@ -680,6 +681,16 @@ bool UContentBrowserAliasDataSource::CanEditItem(const FContentBrowserItemData&
 			}
 			return false;
 		}
+
+		if (UContentBrowserDataSubsystem* ContentBrowserDataSubsystem = IContentBrowserDataModule::Get().GetSubsystem())
+		{
+			const TSharedRef<FPathPermissionList>& EditableFolderFilter = ContentBrowserDataSubsystem->GetEditableFolderPermissionList();
+			if (!EditableFolderFilter->PassesStartsWithFilter(AliasPayload->Alias.Value))
+			{
+				*OutErrorMsg = FText::Format(NSLOCTEXT("ContentBrowserAliasDataSource", "Error_FolderIsNotEditable", "Content in folder '{0}' is not editable"), FText::FromName(AliasPayload->Alias.Value));
+				return false;
+			}
+		}
 	}
 	return ContentBrowserAssetData::CanEditItem(AssetTools, this, InItem, OutErrorMsg);
 }
@@ -739,6 +750,21 @@ bool UContentBrowserAliasDataSource::BulkSaveItems(TArrayView<const FContentBrow
 	return ContentBrowserAssetData::SaveItems(AssetTools, this, InItems, InSaveFlags);
 }
 
+bool UContentBrowserAliasDataSource::CanPrivatizeItem(const FContentBrowserItemData& InItem, FText* OutErrorMsg)
+{
+	return ContentBrowserAssetData::CanPrivatizeItem(AssetTools, AssetRegistry, this, InItem, OutErrorMsg);
+}
+
+bool UContentBrowserAliasDataSource::PrivatizeItem(const FContentBrowserItemData& InItem)
+{
+	return ContentBrowserAssetData::PrivatizeItems(AssetTools, AssetRegistry, this, MakeArrayView(&InItem, 1));
+}
+
+bool UContentBrowserAliasDataSource::BulkPrivatizeItems(TArrayView<const FContentBrowserItemData> InItems)
+{
+	return ContentBrowserAssetData::PrivatizeItems(AssetTools, AssetRegistry, this, InItems);
+}
+
 bool UContentBrowserAliasDataSource::AppendItemReference(const FContentBrowserItemData& InItem, FString& InOutStr)
 {
 	return ContentBrowserAssetData::AppendItemReference(AssetRegistry, this, InItem, InOutStr);
diff --git a/Engine/Plugins/Editor/ContentBrowser/ContentBrowserAliasDataSoure/Source/ContentBrowserAliasDataSource/Public/ContentBrowserAliasDataSource.h b/Engine/Plugins/Editor/ContentBrowser/ContentBrowserAliasDataSoure/Source/ContentBrowserAliasDataSource/Public/ContentBrowserAliasDataSource.h
--- a/Engine/Plugins/Editor/ContentBrowser/ContentBrowserAliasDataSoure/Source/ContentBrowserAliasDataSource/Public/ContentBrowserAliasDataSource.h
+++ b/Engine/Plugins/Editor/ContentBrowser/ContentBrowserAliasDataSoure/Source/ContentBrowserAliasDataSource/Public/ContentBrowserAliasDataSource.h
@@ -83,6 +83,10 @@ public:
 	virtual bool SaveItem(const FContentBrowserItemData& InItem, const EContentBrowserItemSaveFlags InSaveFlags) override;
 	virtual bool BulkSaveItems(TArrayView<const FContentBrowserItemData> InItems, const EContentBrowserItemSaveFlags InSaveFlags) override;
 
+	virtual bool CanPrivatizeItem(const FContentBrowserItemData& InItem, FText* OutErrorMsg) override;
+	virtual bool PrivatizeItem(const FContentBrowserItemData& InItem) override;
+	virtual bool BulkPrivatizeItems(TArrayView<const FContentBrowserItemData> InItems) override;
+
 	virtual bool AppendItemReference(const FContentBrowserItemData& InItem, FString& InOutStr) override;
 	virtual bool UpdateThumbnail(const FContentBrowserItemData& InItem, FAssetThumbnail& InThumbnail) override;
 
diff --git a/Engine/Plugins/Editor/ContentBrowser/ContentBrowserAssetDataSource/Source/ContentBrowserAssetDataSource/Private/ContentBrowserAssetDataCore.cpp b/Engine/Plugins/Editor/ContentBrowser/ContentBrowserAssetDataSource/Source/ContentBrowserAssetDataSource/Private/ContentBrowserAssetDataCore.cpp
--- a/Engine/Plugins/Editor/ContentBrowser/ContentBrowserAssetDataSource/Source/ContentBrowserAssetDataSource/Private/ContentBrowserAssetDataCore.cpp
+++ b/Engine/Plugins/Editor/ContentBrowser/ContentBrowserAssetDataSource/Source/ContentBrowserAssetDataSource/Private/ContentBrowserAssetDataCore.cpp
@@ -5,6 +5,7 @@
 #include "AssetDefinition.h"
 #include "ContentBrowserDataSource.h"
 #include "IAssetTools.h"
+#include "IContentBrowserDataModule.h"
 #include "AssetRegistry/IAssetRegistry.h"
 #include "AssetViewUtils.h"
 #include "AssetTypeActivationOpenedMethod.h"
@@ -202,6 +203,17 @@ bool CanEditItem(IAssetTools* InAssetTools, const UContentBrowserDataSource* InO
 
 bool CanEditAssetFileItem(IAssetTools* InAssetTools, const FContentBrowserAssetFileItemDataPayload& InAssetPayload, FText* OutErrorMsg)
 {
+	if (UContentBrowserDataSubsystem* ContentBrowserDataSubsystem = IContentBrowserDataModule::Get().GetSubsystem())
+	{
+		const TSharedRef<FPathPermissionList>& EditableFolderFilter = ContentBrowserDataSubsystem->GetEditableFolderPermissionList();
+		FName FolderPath = InAssetPayload.GetAssetData().PackageName;
+		if (!EditableFolderFilter->PassesStartsWithFilter(FolderPath))
+		{
+			SetOptionalErrorMessage(OutErrorMsg, FText::Format(LOCTEXT("Error_FolderIsNotEditable", "Content in folder '{0}' is not editable"), FText::FromName(FolderPath)));
+			return false;
+		}
+	}
+
 	if (!CanModifyAssetFileItem(InAssetTools, InAssetPayload, OutErrorMsg))
 	{
 		return false;
diff --git a/Engine/Source/Editor/ContentBrowserData/Private/ContentBrowserDataSubsystem.cpp b/Engine/Source/Editor/ContentBrowserData/Private/ContentBrowserDataSubsystem.cpp
--- a/Engine/Source/Editor/ContentBrowserData/Private/ContentBrowserDataSubsystem.cpp
+++ b/Engine/Source/Editor/ContentBrowserData/Private/ContentBrowserDataSubsystem.cpp
@@ -72,6 +72,11 @@ namespace ContentBrowserDataSubsystem
 	));
 }
 
+UContentBrowserDataSubsystem::UContentBrowserDataSubsystem()
+	: EditableFolderPermissionList(MakeShared<FPathPermissionList>())
+{
+
+}
 void UContentBrowserDataSubsystem::Initialize(FSubsystemCollectionBase& Collection)
 {
 	AllFolderPrefix = TEXT("/All");
@@ -894,6 +899,11 @@ const FString& UContentBrowserDataSubsystem::GetAllFolderPrefix() const
 	return AllFolderPrefix;
 }
 
+TSharedRef<FPathPermissionList>& UContentBrowserDataSubsystem::GetEditableFolderPermissionList()
+{
+	return EditableFolderPermissionList;
+}
+
 EContentBrowserPathType UContentBrowserDataSubsystem::TryConvertVirtualPath(const FStringView InPath, FStringBuilderBase& OutPath) const
 {
 	FNameBuilder FoundVirtualPath;
diff --git a/Engine/Source/Editor/ContentBrowserData/Public/ContentBrowserDataSubsystem.h b/Engine/Source/Editor/ContentBrowserData/Public/ContentBrowserDataSubsystem.h
--- a/Engine/Source/Editor/ContentBrowserData/Public/ContentBrowserDataSubsystem.h
+++ b/Engine/Source/Editor/ContentBrowserData/Public/ContentBrowserDataSubsystem.h
@@ -20,6 +20,7 @@
 #include "EditorSubsystem.h"
 #include "HAL/Platform.h"
 #include "Misc/AssertionMacros.h"
+#include "Misc/NamePermissionList.h"
 #include "Misc/StringBuilder.h"
 #include "UObject/NameTypes.h"
 #include "UObject/ObjectMacros.h"
@@ -91,6 +92,7 @@ class CONTENTBROWSERDATA_API UContentBrowserDataSubsystem : public UEditorSubsys
 public:
 	friend class FScopedSuppressContentBrowserDataTick;
 
+	UContentBrowserDataSubsystem();
 	//~ UEditorSubsystem interface
 	virtual void Initialize(FSubsystemCollectionBase& Collection) override;
 	virtual void Deinitialize() override;
@@ -332,6 +334,12 @@ public:
 	 */
 	const FString& GetAllFolderPrefix() const;
 
+	/**
+	 * Permission list that controls whether content in a given folder path can be edited.
+	 * Note: This does not control if the folder path is writable or if content can be deleted.
+	 */
+	TSharedRef<FPathPermissionList>& GetEditableFolderPermissionList();
+
 private:
 	using FNameToDataSourceMap = TSortedMap<FName, UContentBrowserDataSource*, FDefaultAllocator, FNameFastLess>;
 
@@ -461,6 +469,9 @@ private:
 	 * Prefix to use when generating virtual paths and "Show All Folder" option is enabled
 	 */
 	FString AllFolderPrefix;
+
+	/** Permission list of folder paths we can edit */
+	TSharedRef<FPathPermissionList> EditableFolderPermissionList;
 };
 
 /**
