Fixup read only folder wording

diff --git a/Engine/Plugins/Editor/ContentBrowser/ContentBrowserAliasDataSoure/Source/ContentBrowserAliasDataSource/Private/ContentBrowserAliasDataSource.cpp b/Engine/Plugins/Editor/ContentBrowser/ContentBrowserAliasDataSoure/Source/ContentBrowserAliasDataSource/Private/ContentBrowserAliasDataSource.cpp
--- a/Engine/Plugins/Editor/ContentBrowser/ContentBrowserAliasDataSoure/Source/ContentBrowserAliasDataSource/Private/ContentBrowserAliasDataSource.cpp
+++ b/Engine/Plugins/Editor/ContentBrowser/ContentBrowserAliasDataSoure/Source/ContentBrowserAliasDataSource/Private/ContentBrowserAliasDataSource.cpp
@@ -676,7 +676,7 @@ bool UContentBrowserAliasDataSource::CanEditItem(const FContentBrowserItemData&
 		{
 			if (OutErrorMsg)
 			{
-				*OutErrorMsg = FText::Format(NSLOCTEXT("ContentBrowserAliasDataSource", "Error_FolderIsLocked", "Alias '{0}' is in a locked folder"), FText::FromName(AliasPayload->Alias.Value));
+				*OutErrorMsg = FText::Format(NSLOCTEXT("ContentBrowserAliasDataSource", "Error_FolderIsLocked", "Alias asset '{0}' is in a read only folder. Unable to edit read only assets."), FText::FromName(AliasPayload->Alias.Value));
 			}
 			return false;
 		}
@@ -686,7 +686,10 @@ bool UContentBrowserAliasDataSource::CanEditItem(const FContentBrowserItemData&
 			const TSharedRef<FPathPermissionList>& EditableFolderFilter = ContentBrowserDataSubsystem->GetEditableFolderPermissionList();
 			if (!EditableFolderFilter->PassesStartsWithFilter(AliasPayload->Alias.Value))
 			{
-				*OutErrorMsg = FText::Format(NSLOCTEXT("ContentBrowserAliasDataSource", "Error_FolderIsNotEditable", "Content in folder '{0}' is not editable"), FText::FromName(AliasPayload->Alias.Value));
+				if (OutErrorMsg)
+				{
+					*OutErrorMsg = FText::Format(NSLOCTEXT("ContentBrowserAliasDataSource", "Error_FolderIsNotEditable", "Alias asset '{0}' is in a folder that does not allow edits. Unable to edit read only assets."), FText::FromName(AliasPayload->Alias.Value));
+				}
 				return false;
 			}
 		}
diff --git a/Engine/Plugins/Editor/ContentBrowser/ContentBrowserAssetDataSource/Source/ContentBrowserAssetDataSource/Private/ContentBrowserAssetDataCore.cpp b/Engine/Plugins/Editor/ContentBrowser/ContentBrowserAssetDataSource/Source/ContentBrowserAssetDataSource/Private/ContentBrowserAssetDataCore.cpp
--- a/Engine/Plugins/Editor/ContentBrowser/ContentBrowserAssetDataSource/Source/ContentBrowserAssetDataSource/Private/ContentBrowserAssetDataCore.cpp
+++ b/Engine/Plugins/Editor/ContentBrowser/ContentBrowserAssetDataSource/Source/ContentBrowserAssetDataSource/Private/ContentBrowserAssetDataCore.cpp
@@ -206,10 +206,11 @@ bool CanEditAssetFileItem(IAssetTools* InAssetTools, const FContentBrowserAssetF
 	if (UContentBrowserDataSubsystem* ContentBrowserDataSubsystem = IContentBrowserDataModule::Get().GetSubsystem())
 	{
 		const TSharedRef<FPathPermissionList>& EditableFolderFilter = ContentBrowserDataSubsystem->GetEditableFolderPermissionList();
-		FName FolderPath = InAssetPayload.GetAssetData().PackageName;
-		if (!EditableFolderFilter->PassesStartsWithFilter(FolderPath))
+		FName AssetPackageName = InAssetPayload.GetAssetData().PackageName;
+
+		if (!EditableFolderFilter->PassesStartsWithFilter(AssetPackageName))
 		{
-			SetOptionalErrorMessage(OutErrorMsg, FText::Format(LOCTEXT("Error_FolderIsNotEditable", "Content in folder '{0}' is not editable"), FText::FromName(FolderPath)));
+			SetOptionalErrorMessage(OutErrorMsg, FText::Format(LOCTEXT("Error_FolderIsNotEditable", "Asset '{0}' is in a folder that does not allow edits. Unable to edit read only assets."), FText::FromName(AssetPackageName)));
 			return false;
 		}
 	}
