- Fix `ContentBrowserAssetData::IsPrimaryAsset` to account for Verse classes defined in submodules. This allows the
Content Browser to display Verse creative devices that are defined in submodules.

diff --git a/Engine/Plugins/Editor/ContentBrowser/ContentBrowserAliasDataSoure/Source/ContentBrowserAliasDataSource/Private/ContentBrowserAliasDataSource.cpp b/Engine/Plugins/Editor/ContentBrowser/ContentBrowserAliasDataSoure/Source/ContentBrowserAliasDataSource/Private/ContentBrowserAliasDataSource.cpp
--- a/Engine/Plugins/Editor/ContentBrowser/ContentBrowserAliasDataSoure/Source/ContentBrowserAliasDataSource/Private/ContentBrowserAliasDataSource.cpp
+++ b/Engine/Plugins/Editor/ContentBrowser/ContentBrowserAliasDataSoure/Source/ContentBrowserAliasDataSource/Private/ContentBrowserAliasDataSource.cpp
@@ -305,15 +305,15 @@ TArray<FContentBrowserItemPath> UContentBrowserAliasDataSource::GetAliasesForPat
 	return OutAliases;
 }
 
-void UContentBrowserAliasDataSource::AddAliases(const FAssetData& Asset, const TArray<FName>& Aliases, bool bInIsFromMetaData)
+void UContentBrowserAliasDataSource::AddAliases(const FAssetData& Asset, const TArray<FName>& Aliases, const bool bInIsFromMetaData, const bool bSkipPrimaryAssetValidation)
 {
 	for (const FName Alias : Aliases)
 	{
-		AddAlias(Asset, Alias, bInIsFromMetaData);
+		AddAlias(Asset, Alias, bInIsFromMetaData, bSkipPrimaryAssetValidation);
 	}
 }
 
-void UContentBrowserAliasDataSource::AddAlias(const FAssetData& Asset, const FName Alias, bool bInIsFromMetaData)
+void UContentBrowserAliasDataSource::AddAlias(const FAssetData& Asset, const FName Alias, const bool bInIsFromMetaData, const bool bSkipPrimaryAssetValidation)
 {
 	auto LogErrorMessage = [&Alias, &Asset](const TCHAR* Reason)
 	{
@@ -327,7 +327,7 @@ void UContentBrowserAliasDataSource::AddAlias(const FAssetData& Asset, const FNa
 		return;
 	}
 
-	if (!ContentBrowserAssetData::IsPrimaryAsset(Asset))
+	if (!bSkipPrimaryAssetValidation && !ContentBrowserAssetData::IsPrimaryAsset(Asset))
 	{
 		LogErrorMessage(TEXT("It is not a primary asset"));
 		return;
diff --git a/Engine/Plugins/Editor/ContentBrowser/ContentBrowserAliasDataSoure/Source/ContentBrowserAliasDataSource/Public/ContentBrowserAliasDataSource.h b/Engine/Plugins/Editor/ContentBrowser/ContentBrowserAliasDataSoure/Source/ContentBrowserAliasDataSource/Public/ContentBrowserAliasDataSource.h
--- a/Engine/Plugins/Editor/ContentBrowser/ContentBrowserAliasDataSoure/Source/ContentBrowserAliasDataSource/Public/ContentBrowserAliasDataSource.h
+++ b/Engine/Plugins/Editor/ContentBrowser/ContentBrowserAliasDataSoure/Source/ContentBrowserAliasDataSource/Public/ContentBrowserAliasDataSource.h
@@ -98,10 +98,17 @@ public:
 	virtual bool Legacy_TryConvertAssetDataToVirtualPath(const FAssetData& InAssetData, const bool InUseFolderPaths, FName& OutPath) override;
 	// ~ End UContentBrowserDataSource interface
 
-	/** Add a list of aliases for a given asset. bInIsFromMetaData should only be true if the list of aliases came from the AliasTagName metadata. */
-	void AddAliases(const FAssetData& Asset, const TArray<FName>& Aliases, bool bInIsFromMetaData = false);
+	/** 
+	 * Add a list of aliases for a given asset.
+	 *
+	 * @param Asset 						The asset to add aliases for.
+	 * @param Aliases 						The aliases to add
+	 * @param bInIsFromMetaData 			Should only be `true` if the list of aliases came from the `AliasTagName` metadata.
+	 * @param bSkipPrimaryAssetValidation 	Use this if the asset being added is not a primary asset/is a re-director but should still be included in the Content Browser. (e.g. Verse classes)
+	 */
+	void AddAliases(const FAssetData& Asset, const TArray<FName>& Aliases, const bool bInIsFromMetaData = false, const bool bSkipPrimaryAssetValidation = false);
 	/** Add an alias for a given asset. bInIsFromMetaData should only be true if the alias came from the AliasTagName metadata. */
-	void AddAlias(const FAssetData& Asset, const FName Alias, bool bInIsFromMetaData = false);
+	void AddAlias(const FAssetData& Asset, const FName Alias, const bool bInIsFromMetaData = false, const bool bSkipPrimaryAssetValidation = false);
 	/** Remove the given alias from the data source */
 	void RemoveAlias(const FSoftObjectPath& ObjectPath, const FName Alias);
 	/** Remove all aliases for the given object */
