Fixed aliases failing to refresh their cached asset data (mainly the tags and values) when the underlying asset was loaded/updated

#rnx

diff --git a/Engine/Plugins/Editor/ContentBrowser/ContentBrowserAliasDataSoure/Source/ContentBrowserAliasDataSource/Private/ContentBrowserAliasDataSource.cpp b/Engine/Plugins/Editor/ContentBrowser/ContentBrowserAliasDataSoure/Source/ContentBrowserAliasDataSource/Private/ContentBrowserAliasDataSource.cpp
--- a/Engine/Plugins/Editor/ContentBrowser/ContentBrowserAliasDataSoure/Source/ContentBrowserAliasDataSource/Private/ContentBrowserAliasDataSource.cpp
+++ b/Engine/Plugins/Editor/ContentBrowser/ContentBrowserAliasDataSoure/Source/ContentBrowserAliasDataSource/Private/ContentBrowserAliasDataSource.cpp
@@ -322,7 +322,7 @@ void UContentBrowserAliasDataSource::AddAlias(const FAssetData& Asset, const FNa
 		UE_LOG(LogContentBrowserAliasDataSource, Warning, TEXT("Cannot add alias %s for %s because: %s"), *Alias.ToString(), *Asset.GetObjectPathString(), Reason);
 	};
 	
-	FContentBrowserUniqueAlias UniqueAlias = TPairInitializer<FSoftObjectPath, FName>(Asset.GetSoftObjectPath(), Alias);
+	FContentBrowserUniqueAlias UniqueAlias = MakeTuple(Asset.GetSoftObjectPath(), Alias);
 	if (AllAliases.Contains(UniqueAlias))
 	{
 		LogErrorMessage(TEXT("An alias with that name already exists"));
@@ -428,7 +428,7 @@ PRAGMA_ENABLE_DEPRECATION_WARNINGS
 
 void UContentBrowserAliasDataSource::RemoveAlias(const FSoftObjectPath& ObjectPath, const FName Alias)
 {
-    FContentBrowserUniqueAlias UniqueAlias =  TPairInitializer<FSoftObjectPath, FName>(ObjectPath, Alias);
+    FContentBrowserUniqueAlias UniqueAlias = MakeTuple(ObjectPath, Alias);
 	FAliasData AliasData;
 	// Store a copy of the item data before it's removed for the MakeItemRemovedUpdate notification
 	FContentBrowserItemData ItemData = CreateAssetFileItem(UniqueAlias);
@@ -502,7 +502,7 @@ void UContentBrowserAliasDataSource::ReconcileAliasesFromMetaData(const FAssetDa
 	const TArray<FName>* ExistingAliasesPtr = AliasesForObjectPath.Find(Asset.GetSoftObjectPath());
 	if (ExistingAliasesPtr)
 	{
-		FContentBrowserUniqueAlias UniqueAlias = TPairInitializer<FSoftObjectPath, FName>(Asset.GetSoftObjectPath(), NAME_None);
+		FContentBrowserUniqueAlias UniqueAlias = MakeTuple(Asset.GetSoftObjectPath(), NAME_None);
 		
 		FString AliasTagValue;
 		if (Asset.GetTagValue(AliasTagName, AliasTagValue))
@@ -574,7 +574,7 @@ void UContentBrowserAliasDataSource::ReconcileAliasesForAsset(const FAssetData&
 		if (ExistingAliasesPtr)
 		{
 			TArray<FName> AliasesOnlyInExisting, AliasesOnlyInNew;
-			FContentBrowserUniqueAlias UniqueAlias = TPairInitializer<FSoftObjectPath, FName>(Asset.GetSoftObjectPath(), NAME_None);
+			FContentBrowserUniqueAlias UniqueAlias = MakeTuple(Asset.GetSoftObjectPath(), NAME_None);
 			for (const FName Alias : *ExistingAliasesPtr)
 			{
 				UniqueAlias.Value = Alias;
@@ -611,15 +611,32 @@ void UContentBrowserAliasDataSource::ReconcileAliasesForAsset(const FAssetData&
 void UContentBrowserAliasDataSource::OnAssetUpdated(const FAssetData& InAssetData)
 {
 	ReconcileAliasesFromMetaData(InAssetData);
+	UpdateAliasesCachedAssetData(InAssetData);
+	MakeItemModifiedUpdate(InAssetData.GetSoftObjectPath());
 }
 
-void UContentBrowserAliasDataSource::MakeItemModifiedUpdate(UObject* Object)
+void UContentBrowserAliasDataSource::UpdateAliasesCachedAssetData(const FAssetData& InAssetData)
 {
-	const FSoftObjectPath ObjectPath(Object);
-	const TArray<FName>* Aliases = AliasesForObjectPath.Find(ObjectPath);
-	FContentBrowserUniqueAlias UniqueAlias =  TPairInitializer<FSoftObjectPath, FName>(ObjectPath, NAME_None);
-	if (Aliases)
+	const FSoftObjectPath ObjectPath = InAssetData.GetSoftObjectPath();
+	if (const TArray<FName>* Aliases = AliasesForObjectPath.Find(ObjectPath))
 	{
+		FContentBrowserUniqueAlias UniqueAlias = MakeTuple(ObjectPath, NAME_None);
+		for (const FName Alias : *Aliases)
+		{
+			UniqueAlias.Value = Alias;
+			if (FAliasData* AliasData = AllAliases.Find(UniqueAlias))
+			{
+				AliasData->AssetData = InAssetData;
+			}
+		}
+	}
+}
+
+void UContentBrowserAliasDataSource::MakeItemModifiedUpdate(const FSoftObjectPath& ObjectPath)
+{
+	if (const TArray<FName>* Aliases = AliasesForObjectPath.Find(ObjectPath))
+	{
+		FContentBrowserUniqueAlias UniqueAlias = MakeTuple(ObjectPath, NAME_None);
 		for (const FName Alias : *Aliases)
 		{
 			UniqueAlias.Value = Alias;
@@ -632,14 +649,19 @@ void UContentBrowserAliasDataSource::OnAssetLoaded(UObject* InAsset)
 {
 	if (InAsset && !InAsset->GetOutermost()->HasAnyPackageFlags(PKG_ForDiffing))
 	{
+		const FAssetData LoadedAssetData(InAsset);
+		ReconcileAliasesFromMetaData(LoadedAssetData);
+		UpdateAliasesCachedAssetData(LoadedAssetData);
 		MakeItemModifiedUpdate(InAsset);
 	}
 }
 
 void UContentBrowserAliasDataSource::OnObjectPropertyChanged(UObject* InObject, FPropertyChangedEvent& InPropertyChangedEvent)
 {
-	if (InObject && InObject->IsAsset())
+	if (InObject && InObject->IsAsset() && AliasesForObjectPath.Contains(InObject))
 	{
+		const FAssetData LoadedAssetData(InObject);
+		UpdateAliasesCachedAssetData(LoadedAssetData);
 		MakeItemModifiedUpdate(InObject);
 	}
 }
diff --git a/Engine/Plugins/Editor/ContentBrowser/ContentBrowserAliasDataSoure/Source/ContentBrowserAliasDataSource/Public/ContentBrowserAliasDataSource.h b/Engine/Plugins/Editor/ContentBrowser/ContentBrowserAliasDataSoure/Source/ContentBrowserAliasDataSource/Public/ContentBrowserAliasDataSource.h
--- a/Engine/Plugins/Editor/ContentBrowser/ContentBrowserAliasDataSoure/Source/ContentBrowserAliasDataSource/Public/ContentBrowserAliasDataSource.h
+++ b/Engine/Plugins/Editor/ContentBrowser/ContentBrowserAliasDataSoure/Source/ContentBrowserAliasDataSource/Public/ContentBrowserAliasDataSource.h
@@ -138,7 +138,9 @@ private:
 
 	/** Helper function to remove a folder from the PathTree including all parent folders that are now empty as a result of the removal */
 	void RemoveFoldersRecursively(FStringView LeafFolder);
-	void MakeItemModifiedUpdate(UObject* Object);
+
+	void UpdateAliasesCachedAssetData(const FAssetData& InAssetData);
+	void MakeItemModifiedUpdate(const FSoftObjectPath& ObjectPath);
 
 	FContentBrowserItemData CreateAssetFolderItem(const FName InFolderPath);
 	FContentBrowserItemData CreateAssetFileItem(const FContentBrowserUniqueAlias& Alias);
