ContentBrowser.LogAliases command to log aliases
RemoveAlias optims

diff --git a/Engine/Plugins/Editor/ContentBrowser/ContentBrowserAliasDataSoure/Source/ContentBrowserAliasDataSource/Private/ContentBrowserAliasDataSource.cpp b/Engine/Plugins/Editor/ContentBrowser/ContentBrowserAliasDataSoure/Source/ContentBrowserAliasDataSource/Private/ContentBrowserAliasDataSource.cpp
--- a/Engine/Plugins/Editor/ContentBrowser/ContentBrowserAliasDataSoure/Source/ContentBrowserAliasDataSource/Private/ContentBrowserAliasDataSource.cpp
+++ b/Engine/Plugins/Editor/ContentBrowser/ContentBrowserAliasDataSoure/Source/ContentBrowserAliasDataSource/Private/ContentBrowserAliasDataSource.cpp
@@ -1,6 +1,7 @@
 // Copyright Epic Games, Inc. All Rights Reserved.
 
 #include "ContentBrowserAliasDataSource.h"
+#include "ContentBrowserAliasDataSourceModule.h"
 
 #include "AssetRegistry/AssetRegistryModule.h"
 #include "ContentBrowserAssetDataCore.h"
@@ -17,6 +18,21 @@
 
 DEFINE_LOG_CATEGORY(LogContentBrowserAliasDataSource);
 
+namespace ContentBrowserAliasDataSource
+{
+	FAutoConsoleCommand LogAliasesCommand = FAutoConsoleCommand(
+		TEXT("ContentBrowser.LogAliases"),
+		TEXT("Logs all content browser aliases"),
+		FConsoleCommandWithArgsDelegate::CreateLambda(
+			[](const TArray<FString>& /*Args*/)
+			{
+				if (UContentBrowserAliasDataSource* AliasDataSource = FModuleManager::Get().LoadModuleChecked<FContentBrowserAliasDataSourceModule>("ContentBrowserAliasDataSource").GetAliasDataSource().Get())
+				{
+					AliasDataSource->LogAliases();
+				}
+			}));
+}
+
 FName UContentBrowserAliasDataSource::AliasTagName = "ContentBrowserAliases";
 
 void UContentBrowserAliasDataSource::Initialize(const bool InAutoRegister)
@@ -419,16 +435,23 @@ void UContentBrowserAliasDataSource::RemoveAlias(const FSoftObjectPath& ObjectPa
 	FContentBrowserItemData ItemData = CreateAssetFileItem(UniqueAlias);
 	if (AllAliases.RemoveAndCopyValue(UniqueAlias, AliasData))
 	{
-		AliasesForObjectPath[AliasData.AssetData.GetSoftObjectPath()].Remove(Alias);
-		if (AliasesForObjectPath[AliasData.AssetData.GetSoftObjectPath()].Num() == 0)
+		check(AliasData.AssetData.GetSoftObjectPath() == ObjectPath);
 		{
-			AliasesForObjectPath.Remove(AliasData.AssetData.GetSoftObjectPath());
+			TArray<FName>& Aliases = AliasesForObjectPath[ObjectPath];
+			Aliases.Remove(Alias);
+			if (Aliases.IsEmpty())
+			{
+				AliasesForObjectPath.Remove(ObjectPath);
+			}
 		}
 
-		AliasesInPackagePath[AliasData.PackagePath].Remove(UniqueAlias);
-		if (AliasesInPackagePath[AliasData.PackagePath].Num() == 0)
 		{
-			AliasesInPackagePath.Remove(AliasData.PackagePath);
+			TArray<FContentBrowserUniqueAlias>& Aliases = AliasesInPackagePath[AliasData.PackagePath];
+			Aliases.Remove(UniqueAlias);
+			if (Aliases.IsEmpty())
+			{
+				AliasesInPackagePath.Remove(AliasData.PackagePath);
+			}
 		}
 
 		QueueItemDataUpdate(FContentBrowserItemDataUpdate::MakeItemRemovedUpdate(ItemData));
@@ -593,6 +616,28 @@ void UContentBrowserAliasDataSource::ReconcileAliasesForAsset(const FAssetData&
 	}
 }
 
+void UContentBrowserAliasDataSource::LogAliases() const
+{
+	TArray<FContentBrowserUniqueAlias> Aliases;
+	AllAliases.GenerateKeyArray(Aliases);
+	Aliases.Sort(
+		[](const FContentBrowserUniqueAlias& A, const FContentBrowserUniqueAlias& B)
+		{
+			return A.Key == B.Key ? A.Value.LexicalLess(B.Value) : A.Key.LexicalLess(B.Key);
+		});
+
+	TStringBuilder<FName::StringBufferSize * 2> StringBuilder;
+	for (const FContentBrowserUniqueAlias& Alias : Aliases)
+	{
+		StringBuilder.Reset();
+		StringBuilder.Append(Alias.Key.ToString());
+		StringBuilder.Append(TEXT(" -> \""));
+		StringBuilder.Append(Alias.Value.ToString());
+		StringBuilder.AppendChar(TCHAR('"'));
+		UE_LOG(LogContentBrowserAliasDataSource, Log, TEXT("%s"), StringBuilder.ToString());
+	}
+}
+
 void UContentBrowserAliasDataSource::OnAssetUpdated(const FAssetData& InAssetData)
 {
 	ReconcileAliasesFromMetaData(InAssetData);
diff --git a/Engine/Plugins/Editor/ContentBrowser/ContentBrowserAliasDataSoure/Source/ContentBrowserAliasDataSource/Public/ContentBrowserAliasDataSource.h b/Engine/Plugins/Editor/ContentBrowser/ContentBrowserAliasDataSoure/Source/ContentBrowserAliasDataSource/Public/ContentBrowserAliasDataSource.h
--- a/Engine/Plugins/Editor/ContentBrowser/ContentBrowserAliasDataSoure/Source/ContentBrowserAliasDataSource/Public/ContentBrowserAliasDataSource.h
+++ b/Engine/Plugins/Editor/ContentBrowser/ContentBrowserAliasDataSoure/Source/ContentBrowserAliasDataSource/Public/ContentBrowserAliasDataSource.h
@@ -131,6 +131,9 @@ public:
 	/** Calls AddAlias or RemoveAlias for every alias that doesn't match the stored data for the given asset. */
 	void ReconcileAliasesForAsset(const FAssetData& Asset, const TArray<FName>& NewAliases);
 
+	/** Logs all the content browser aliases */
+	void LogAliases() const;
+
 protected:
 	virtual void BuildRootPathVirtualTree() override;
 
