Fix to ensure uniqueness of content browser items for two dffierent assets with the same name sharing the same alias

diff --git a/Engine/Plugins/Editor/ContentBrowser/ContentBrowserAliasDataSoure/Source/ContentBrowserAliasDataSource/Public/ContentBrowserAliasDataSource.h b/Engine/Plugins/Editor/ContentBrowser/ContentBrowserAliasDataSoure/Source/ContentBrowserAliasDataSource/Public/ContentBrowserAliasDataSource.h
--- a/Engine/Plugins/Editor/ContentBrowser/ContentBrowserAliasDataSoure/Source/ContentBrowserAliasDataSource/Public/ContentBrowserAliasDataSource.h
+++ b/Engine/Plugins/Editor/ContentBrowser/ContentBrowserAliasDataSoure/Source/ContentBrowserAliasDataSource/Public/ContentBrowserAliasDataSource.h
@@ -162,12 +162,21 @@ private:
 		FAliasData(const FAssetData& InAssetData, const FName InPackagePath, const FName InName, const bool bInIsFromMetaData = false)
 			: AssetData(InAssetData), PackagePath(InPackagePath), AliasName(InName), bIsFromMetaData(bInIsFromMetaData)
 		{
-			FNameBuilder PathBuilder(PackagePath);
-			PathBuilder << TEXT('/');
-			PathBuilder << InAssetData.AssetName;
-			PackageName = FName(PathBuilder.ToView());
-
-			ObjectPath = FSoftObjectPath(PackageName, InAssetData.AssetName, {});
+			FNameBuilder AssetNameBuilder(InAssetData.AssetName);
+			{
+				// Add a hash of the real asset path to the asset name to ensure uniqueness of the FContentBrowserItemKey
+				// so that two different assets with the same name and the same alias both show up in the content browser
+				FNameBuilder AssetPathBuilder;
+				InAssetData.AppendObjectPath(AssetPathBuilder);
+				AssetNameBuilder.Appendf(TEXT("_%08X"), GetTypeHash(AssetPathBuilder.ToView()));
+			}
+			{
+				FNameBuilder PathBuilder(PackagePath);
+				PathBuilder << TEXT('/');
+				PathBuilder << AssetNameBuilder.ToView();
+				PackageName = FName(PathBuilder.ToView());
+			}
+			ObjectPath = FSoftObjectPath(PackageName, FName(AssetNameBuilder.ToView()), {});
 		}
 
 		/** The source asset for this alias */
