ContentBrowser now allows rebuildling aliases
#rnx
#tests

diff --git a/Engine/Plugins/Editor/ContentBrowser/ContentBrowserAliasDataSoure/Source/ContentBrowserAliasDataSource/Private/ContentBrowserAliasDataSource.cpp b/Engine/Plugins/Editor/ContentBrowser/ContentBrowserAliasDataSoure/Source/ContentBrowserAliasDataSource/Private/ContentBrowserAliasDataSource.cpp
--- a/Engine/Plugins/Editor/ContentBrowser/ContentBrowserAliasDataSoure/Source/ContentBrowserAliasDataSource/Private/ContentBrowserAliasDataSource.cpp
+++ b/Engine/Plugins/Editor/ContentBrowser/ContentBrowserAliasDataSoure/Source/ContentBrowserAliasDataSource/Private/ContentBrowserAliasDataSource.cpp
@@ -465,6 +465,19 @@ void UContentBrowserAliasDataSource::RemoveAlias(const FSoftObjectPath& ObjectPa
 	}
 }
 
+void UContentBrowserAliasDataSource::RebuildAliases()
+{
+	PathTree = FPathTree();
+	AllAliases.Reset();
+	AliasesForObjectPath.Reset();
+	AliasesInPackagePath.Reset();
+	AlreadyAddedOriginalAssets.Reset();
+	FilterCache.Reset();
+	RootPathVirtualTree.Reset();
+
+	OnRebuildAliases().Broadcast();
+}
+
 PRAGMA_DISABLE_DEPRECATION_WARNINGS
 void UContentBrowserAliasDataSource::RemoveAliases(FName ObjectPath)
 {
diff --git a/Engine/Plugins/Editor/ContentBrowser/ContentBrowserAliasDataSoure/Source/ContentBrowserAliasDataSource/Public/ContentBrowserAliasDataSource.h b/Engine/Plugins/Editor/ContentBrowser/ContentBrowserAliasDataSoure/Source/ContentBrowserAliasDataSource/Public/ContentBrowserAliasDataSource.h
--- a/Engine/Plugins/Editor/ContentBrowser/ContentBrowserAliasDataSoure/Source/ContentBrowserAliasDataSource/Public/ContentBrowserAliasDataSource.h
+++ b/Engine/Plugins/Editor/ContentBrowser/ContentBrowserAliasDataSoure/Source/ContentBrowserAliasDataSource/Public/ContentBrowserAliasDataSource.h
@@ -126,6 +126,12 @@ public:
 	UE_DEPRECATED(5.1, "FNames containing full asset paths are deprecated, use FSoftObjectPath instead")
 	void RemoveAliases(const FName ObjectPath);
 
+	/** When called, removes all aliases and triggers delegate for various systems to re-add aliases */
+	void RebuildAliases();
+
+	/** Broadcast after RebuildAliases() called to allow systems to re-add aliases */
+	FSimpleMulticastDelegate& OnRebuildAliases() { return RebuildAliasesDelegate; }
+
 	/** Get all aliases from metadata for the given asset, then calls AddAlias or RemoveAlias for every alias that doesn't match the stored data. */
 	void ReconcileAliasesFromMetaData(const FAssetData& Asset);
 	/** Calls AddAlias or RemoveAlias for every alias that doesn't match the stored data for the given asset. */
@@ -206,6 +212,9 @@ private:
 	TSet<FSoftObjectPath> AlreadyAddedOriginalAssets;
 
 	UContentBrowserAssetDataSource::FAssetDataSourceFilterCache FilterCache;
+
+	/** Delegate broadcast after all aliases removed to give chance for systems to re-add aliases */
+	FSimpleMulticastDelegate RebuildAliasesDelegate;
 };
 
 #if UE_ENABLE_INCLUDE_ORDER_DEPRECATED_IN_5_2
diff --git a/Engine/Plugins/Editor/ContentBrowser/ContentBrowserAssetDataSource/Source/ContentBrowserAssetDataSource/Private/ContentBrowserAssetDataSource.cpp b/Engine/Plugins/Editor/ContentBrowser/ContentBrowserAssetDataSource/Source/ContentBrowserAssetDataSource/Private/ContentBrowserAssetDataSource.cpp
--- a/Engine/Plugins/Editor/ContentBrowser/ContentBrowserAssetDataSource/Source/ContentBrowserAssetDataSource/Private/ContentBrowserAssetDataSource.cpp
+++ b/Engine/Plugins/Editor/ContentBrowser/ContentBrowserAssetDataSource/Source/ContentBrowserAssetDataSource/Private/ContentBrowserAssetDataSource.cpp
@@ -3463,6 +3463,11 @@ void UContentBrowserAssetDataSource::FAssetDataSourceFilterCache::ClearCachedDat
 	CachedCompiledInternalPaths.Remove(IDOwner);
 }
 
+void UContentBrowserAssetDataSource::FAssetDataSourceFilterCache::Reset()
+{
+	CachedCompiledInternalPaths.Reset(); 
+}
+
 void UContentBrowserAssetDataSource::FAssetDataSourceFilterCache::OnPathAdded(FName Path, FStringView PathString, uint32 PathHash, FName ParentPath, uint32 ParentPathHash, int32 PathDepth)
 {
 	for (TPair<FContentBrowserDataFilterCacheID, FCachedDataPerID>& CachedCompiledInternalPath : CachedCompiledInternalPaths)
diff --git a/Engine/Plugins/Editor/ContentBrowser/ContentBrowserAssetDataSource/Source/ContentBrowserAssetDataSource/Public/ContentBrowserAssetDataSource.h b/Engine/Plugins/Editor/ContentBrowser/ContentBrowserAssetDataSource/Source/ContentBrowserAssetDataSource/Public/ContentBrowserAssetDataSource.h
--- a/Engine/Plugins/Editor/ContentBrowser/ContentBrowserAssetDataSource/Source/ContentBrowserAssetDataSource/Public/ContentBrowserAssetDataSource.h
+++ b/Engine/Plugins/Editor/ContentBrowser/ContentBrowserAssetDataSource/Source/ContentBrowserAssetDataSource/Public/ContentBrowserAssetDataSource.h
@@ -133,6 +133,7 @@ public:
 
 		void RemoveUnusedCachedData(const FContentBrowserDataFilterCacheIDOwner& IDOwner, TArrayView<const FName> InVirtualPathsInUse, const FContentBrowserDataFilter& DataFilter);
 		void ClearCachedData(const FContentBrowserDataFilterCacheIDOwner& IDOwner);
+		void Reset();
 
 	private:
 		friend UContentBrowserAssetDataSource;
