Fix crash on editor shutdown in alias manager accessing content browser alias data source after exit purge GC
#rnx
#tests

diff --git a/Engine/Plugins/Editor/ContentBrowser/ContentBrowserAliasDataSoure/Source/ContentBrowserAliasDataSource/Private/ContentBrowserAliasDataSource.cpp b/Engine/Plugins/Editor/ContentBrowser/ContentBrowserAliasDataSoure/Source/ContentBrowserAliasDataSource/Private/ContentBrowserAliasDataSource.cpp
--- a/Engine/Plugins/Editor/ContentBrowser/ContentBrowserAliasDataSoure/Source/ContentBrowserAliasDataSource/Private/ContentBrowserAliasDataSource.cpp
+++ b/Engine/Plugins/Editor/ContentBrowser/ContentBrowserAliasDataSoure/Source/ContentBrowserAliasDataSource/Private/ContentBrowserAliasDataSource.cpp
@@ -26,7 +26,7 @@ namespace ContentBrowserAliasDataSource
 		FConsoleCommandWithArgsDelegate::CreateLambda(
 			[](const TArray<FString>& /*Args*/)
 			{
-				if (UContentBrowserAliasDataSource* AliasDataSource = FModuleManager::Get().LoadModuleChecked<FContentBrowserAliasDataSourceModule>("ContentBrowserAliasDataSource").GetAliasDataSource().Get())
+				if (UContentBrowserAliasDataSource* AliasDataSource = FModuleManager::Get().LoadModuleChecked<FContentBrowserAliasDataSourceModule>("ContentBrowserAliasDataSource").GetAliasDataSource())
 				{
 					AliasDataSource->LogAliases();
 				}
diff --git a/Engine/Plugins/Editor/ContentBrowser/ContentBrowserAliasDataSoure/Source/ContentBrowserAliasDataSource/Private/ContentBrowserAliasDataSourceModule.cpp b/Engine/Plugins/Editor/ContentBrowser/ContentBrowserAliasDataSoure/Source/ContentBrowserAliasDataSource/Private/ContentBrowserAliasDataSourceModule.cpp
--- a/Engine/Plugins/Editor/ContentBrowser/ContentBrowserAliasDataSoure/Source/ContentBrowserAliasDataSource/Private/ContentBrowserAliasDataSourceModule.cpp
+++ b/Engine/Plugins/Editor/ContentBrowser/ContentBrowserAliasDataSoure/Source/ContentBrowserAliasDataSource/Private/ContentBrowserAliasDataSourceModule.cpp
@@ -10,9 +10,26 @@ void FContentBrowserAliasDataSourceModule::StartupModule()
 	AliasDataSource->Initialize();
 }
 
-void FContentBrowserAliasDataSourceModule::ShutdownModule()
+void FContentBrowserAliasDataSourceModule::PreUnloadCallback()
 {
 	AliasDataSource.Reset();
 }
 
+UContentBrowserAliasDataSource* FContentBrowserAliasDataSourceModule::GetAliasDataSource()
+{
+	// GExitPurge guard required because crash inspecting object after GExitPurge
+	if (!GExitPurge)
+	{
+		if (UContentBrowserAliasDataSource* Obj = AliasDataSource.Get())
+		{
+			if (IsValid(Obj))
+			{
+				return Obj;
+			}
+		}
+	}
+
+	return nullptr;
+}
+
 IMPLEMENT_MODULE(FContentBrowserAliasDataSourceModule, ContentBrowserAliasDataSource);
diff --git a/Engine/Plugins/Editor/ContentBrowser/ContentBrowserAliasDataSoure/Source/ContentBrowserAliasDataSource/Public/ContentBrowserAliasDataSourceModule.h b/Engine/Plugins/Editor/ContentBrowser/ContentBrowserAliasDataSoure/Source/ContentBrowserAliasDataSource/Public/ContentBrowserAliasDataSourceModule.h
--- a/Engine/Plugins/Editor/ContentBrowser/ContentBrowserAliasDataSoure/Source/ContentBrowserAliasDataSource/Public/ContentBrowserAliasDataSourceModule.h
+++ b/Engine/Plugins/Editor/ContentBrowser/ContentBrowserAliasDataSoure/Source/ContentBrowserAliasDataSource/Public/ContentBrowserAliasDataSourceModule.h
@@ -10,9 +10,9 @@ class FContentBrowserAliasDataSourceModule : public FDefaultModuleImpl
 {
 public:
 	virtual void StartupModule() override;
-	virtual void ShutdownModule() override;
+	virtual void PreUnloadCallback() override;
 
-	TWeakObjectPtr<UContentBrowserAliasDataSource> GetAliasDataSource() { return AliasDataSource.Get(); }
+	CONTENTBROWSERALIASDATASOURCE_API UContentBrowserAliasDataSource* GetAliasDataSource();
 
 private:
 	TStrongObjectPtr<UContentBrowserAliasDataSource> AliasDataSource;
