Add back deprecated version of GetAliasDataSource, new function is now named TryGetAliasDataSource
#rnx

diff --git a/Engine/Plugins/Editor/ContentBrowser/ContentBrowserAliasDataSoure/Source/ContentBrowserAliasDataSource/Private/ContentBrowserAliasDataSource.cpp b/Engine/Plugins/Editor/ContentBrowser/ContentBrowserAliasDataSoure/Source/ContentBrowserAliasDataSource/Private/ContentBrowserAliasDataSource.cpp
--- a/Engine/Plugins/Editor/ContentBrowser/ContentBrowserAliasDataSoure/Source/ContentBrowserAliasDataSource/Private/ContentBrowserAliasDataSource.cpp
+++ b/Engine/Plugins/Editor/ContentBrowser/ContentBrowserAliasDataSoure/Source/ContentBrowserAliasDataSource/Private/ContentBrowserAliasDataSource.cpp
@@ -26,7 +26,7 @@ namespace ContentBrowserAliasDataSource
 		FConsoleCommandWithArgsDelegate::CreateLambda(
 			[](const TArray<FString>& /*Args*/)
 			{
-				if (UContentBrowserAliasDataSource* AliasDataSource = FModuleManager::Get().LoadModuleChecked<FContentBrowserAliasDataSourceModule>("ContentBrowserAliasDataSource").GetAliasDataSource())
+				if (UContentBrowserAliasDataSource* AliasDataSource = FModuleManager::Get().LoadModuleChecked<FContentBrowserAliasDataSourceModule>("ContentBrowserAliasDataSource").TryGetAliasDataSource())
 				{
 					AliasDataSource->LogAliases();
 				}
diff --git a/Engine/Plugins/Editor/ContentBrowser/ContentBrowserAliasDataSoure/Source/ContentBrowserAliasDataSource/Private/ContentBrowserAliasDataSourceModule.cpp b/Engine/Plugins/Editor/ContentBrowser/ContentBrowserAliasDataSoure/Source/ContentBrowserAliasDataSource/Private/ContentBrowserAliasDataSourceModule.cpp
--- a/Engine/Plugins/Editor/ContentBrowser/ContentBrowserAliasDataSoure/Source/ContentBrowserAliasDataSource/Private/ContentBrowserAliasDataSourceModule.cpp
+++ b/Engine/Plugins/Editor/ContentBrowser/ContentBrowserAliasDataSoure/Source/ContentBrowserAliasDataSource/Private/ContentBrowserAliasDataSourceModule.cpp
@@ -15,7 +15,7 @@ void FContentBrowserAliasDataSourceModule::PreUnloadCallback()
 	AliasDataSource.Reset();
 }
 
-UContentBrowserAliasDataSource* FContentBrowserAliasDataSourceModule::GetAliasDataSource()
+UContentBrowserAliasDataSource* FContentBrowserAliasDataSourceModule::TryGetAliasDataSource()
 {
 	// GExitPurge guard required because crash inspecting object after GExitPurge
 	if (!GExitPurge)
diff --git a/Engine/Plugins/Editor/ContentBrowser/ContentBrowserAliasDataSoure/Source/ContentBrowserAliasDataSource/Public/ContentBrowserAliasDataSourceModule.h b/Engine/Plugins/Editor/ContentBrowser/ContentBrowserAliasDataSoure/Source/ContentBrowserAliasDataSource/Public/ContentBrowserAliasDataSourceModule.h
--- a/Engine/Plugins/Editor/ContentBrowser/ContentBrowserAliasDataSoure/Source/ContentBrowserAliasDataSource/Public/ContentBrowserAliasDataSourceModule.h
+++ b/Engine/Plugins/Editor/ContentBrowser/ContentBrowserAliasDataSoure/Source/ContentBrowserAliasDataSource/Public/ContentBrowserAliasDataSourceModule.h
@@ -12,7 +12,10 @@ public:
 	virtual void StartupModule() override;
 	virtual void PreUnloadCallback() override;
 
-	CONTENTBROWSERALIASDATASOURCE_API UContentBrowserAliasDataSource* GetAliasDataSource();
+	UE_DEPRECATED(5.4, "Call TryGetAliasDataSource instead.")
+	TWeakObjectPtr<UContentBrowserAliasDataSource> GetAliasDataSource() { return AliasDataSource.Get(); }
+
+	CONTENTBROWSERALIASDATASOURCE_API UContentBrowserAliasDataSource* TryGetAliasDataSource();
 
 private:
 	TStrongObjectPtr<UContentBrowserAliasDataSource> AliasDataSource;
