Added support for localized content aliases

This allows you to specify a display name override for an alias or the path to an alias

diff --git a/Engine/Plugins/Editor/ContentBrowser/ContentBrowserAliasDataSoure/Source/ContentBrowserAliasDataSource/Private/ContentBrowserAliasDataSource.cpp b/Engine/Plugins/Editor/ContentBrowser/ContentBrowserAliasDataSoure/Source/ContentBrowserAliasDataSource/Private/ContentBrowserAliasDataSource.cpp
--- a/Engine/Plugins/Editor/ContentBrowser/ContentBrowserAliasDataSoure/Source/ContentBrowserAliasDataSource/Private/ContentBrowserAliasDataSource.cpp
+++ b/Engine/Plugins/Editor/ContentBrowser/ContentBrowserAliasDataSoure/Source/ContentBrowserAliasDataSource/Private/ContentBrowserAliasDataSource.cpp
@@ -31,6 +31,36 @@ namespace ContentBrowserAliasDataSource
 					AliasDataSource->LogAliases();
 				}
 			}));
+
+	template <typename AliasType>
+	FName GetAliasName(const AliasType&);
+
+	template <>
+	FName GetAliasName<FName>(const FName& Alias)
+	{
+		return Alias;
+	}
+
+	template <>
+	FName GetAliasName<FContentBrowserLocalizedAlias>(const FContentBrowserLocalizedAlias& Alias)
+	{
+		return Alias.Alias;
+	}
+
+	template <typename AliasType>
+	FText GetAliasDisplayNameOverride(const AliasType&);
+
+	template <>
+	FText GetAliasDisplayNameOverride<FName>(const FName& Alias)
+	{
+		return FText::GetEmpty();
+	}
+
+	template <>
+	FText GetAliasDisplayNameOverride<FContentBrowserLocalizedAlias>(const FContentBrowserLocalizedAlias& Alias)
+	{
+		return Alias.DisplayName;
+	}
 }
 
 FName UContentBrowserAliasDataSource::AliasTagName = "ContentBrowserAliases";
@@ -310,7 +340,18 @@ TArray<FContentBrowserItemPath> UContentBrowserAliasDataSource::GetAliasesForPat
 
 void UContentBrowserAliasDataSource::AddAliases(const FAssetData& Asset, const TArray<FName>& Aliases, const bool bInIsFromMetaData, const bool bSkipPrimaryAssetValidation)
 {
-	for (const FName Alias : Aliases)
+	AddAliasesImpl(Asset, Aliases, bInIsFromMetaData, bSkipPrimaryAssetValidation);
+}
+
+void UContentBrowserAliasDataSource::AddAliases(const FAssetData& Asset, const TArray<FContentBrowserLocalizedAlias>& Aliases, const bool bInIsFromMetaData, const bool bSkipPrimaryAssetValidation)
+{
+	AddAliasesImpl(Asset, Aliases, bInIsFromMetaData, bSkipPrimaryAssetValidation);
+}
+
+template <typename AliasType>
+void UContentBrowserAliasDataSource::AddAliasesImpl(const FAssetData& Asset, const TArray<AliasType>& Aliases, const bool bInIsFromMetaData, const bool bSkipPrimaryAssetValidation)
+{
+	for (const AliasType& Alias : Aliases)
 	{
 		AddAlias(Asset, Alias, bInIsFromMetaData, bSkipPrimaryAssetValidation);
 	}
@@ -318,6 +359,20 @@ void UContentBrowserAliasDataSource::AddAliases(const FAssetData& Asset, const T
 
 void UContentBrowserAliasDataSource::AddAlias(const FAssetData& Asset, const FName Alias, const bool bInIsFromMetaData, const bool bSkipPrimaryAssetValidation)
 {
+	AddAliasImpl(Asset, Alias, bInIsFromMetaData, bSkipPrimaryAssetValidation);
+}
+
+void UContentBrowserAliasDataSource::AddAlias(const FAssetData& Asset, const FContentBrowserLocalizedAlias& Alias, const bool bInIsFromMetaData, const bool bSkipPrimaryAssetValidation)
+{
+	AddAliasImpl(Asset, Alias, bInIsFromMetaData, bSkipPrimaryAssetValidation);
+}
+
+template <typename AliasType>
+void UContentBrowserAliasDataSource::AddAliasImpl(const FAssetData& Asset, const AliasType& AliasToAdd, const bool bInIsFromMetaData, const bool bSkipPrimaryAssetValidation)
+{
+	const FName Alias = ContentBrowserAliasDataSource::GetAliasName(AliasToAdd);
+	const FText AliasDisplayNameOverride = ContentBrowserAliasDataSource::GetAliasDisplayNameOverride(AliasToAdd);
+
 	auto LogErrorMessage = [&Alias, &Asset](const TCHAR* Reason)
 	{
 		UE_LOG(LogContentBrowserAliasDataSource, Warning, TEXT("Cannot add alias %s for %s because: %s"), *Alias.ToString(), *Asset.GetObjectPathString(), Reason);
@@ -376,7 +431,7 @@ void UContentBrowserAliasDataSource::AddAlias(const FAssetData& Asset, const FNa
 
 	AliasesInPackagePath.FindOrAdd(PackagePath).Add(UniqueAlias);
 	AliasesForObjectPath.FindOrAdd(Asset.GetSoftObjectPath()).Add(Alias);
-	AllAliases.Add(UniqueAlias, FAliasData(Asset, PackagePath, *Tokens.Last(), bInIsFromMetaData));
+	AllAliases.Add(UniqueAlias, FAliasData(Asset, PackagePath, AliasDisplayNameOverride.IsEmpty() ? FText::AsCultureInvariant(*Tokens.Last()) : AliasDisplayNameOverride, bInIsFromMetaData));
 	QueueItemDataUpdate(FContentBrowserItemDataUpdate::MakeItemAddedUpdate(CreateAssetFileItem(UniqueAlias)));
 
 	// This logging might get out of control if there ends up being hundreds of thousands of aliases.
@@ -465,12 +520,39 @@ void UContentBrowserAliasDataSource::RemoveAlias(const FSoftObjectPath& ObjectPa
 	}
 }
 
+void UContentBrowserAliasDataSource::AddAliasFolderDisplayName(const FName AliasFolder, const FText& DisplayName)
+{
+	if (const FText* ExistingDisplayName = AliasFolderDisplayNames.Find(AliasFolder);
+		ExistingDisplayName && ExistingDisplayName->IdenticalTo(DisplayName, ETextIdenticalModeFlags::LexicalCompareInvariants))
+	{
+		return;
+	}
+
+	AliasFolderDisplayNames.Add(AliasFolder, DisplayName);
+
+	if (PathTree.PathExists(AliasFolder))
+	{
+		QueueItemDataUpdate(FContentBrowserItemDataUpdate::MakeItemModifiedUpdate(CreateAssetFolderItem(AliasFolder)));
+	}
+}
+
+void UContentBrowserAliasDataSource::RemoveAliasFolderDisplayName(const FName AliasFolder)
+{
+	AliasFolderDisplayNames.Remove(AliasFolder);
+
+	if (PathTree.PathExists(AliasFolder))
+	{
+		QueueItemDataUpdate(FContentBrowserItemDataUpdate::MakeItemModifiedUpdate(CreateAssetFolderItem(AliasFolder)));
+	}
+}
+
 void UContentBrowserAliasDataSource::RebuildAliases()
 {
 	PathTree = FPathTree();
 	AllAliases.Reset();
 	AliasesForObjectPath.Reset();
 	AliasesInPackagePath.Reset();
+	AliasFolderDisplayNames.Reset();
 	AlreadyAddedOriginalAssets.Reset();
 	FilterCache.Reset();
 	RootPathVirtualTree.Reset();
@@ -588,29 +670,41 @@ void UContentBrowserAliasDataSource::ReconcileAliasesFromMetaData(const FAssetDa
 }
 
 void UContentBrowserAliasDataSource::ReconcileAliasesForAsset(const FAssetData& Asset, const TArray<FName>& NewAliases)
+{
+	ReconcileAliasesForAssetImpl(Asset, NewAliases);
+}
+
+void UContentBrowserAliasDataSource::ReconcileAliasesForAsset(const FAssetData& Asset, const TArray<FContentBrowserLocalizedAlias>& NewAliases)
+{
+	ReconcileAliasesForAssetImpl(Asset, NewAliases);
+}
+
+template <typename AliasType>
+void UContentBrowserAliasDataSource::ReconcileAliasesForAssetImpl(const FAssetData& Asset, const TArray<AliasType>& NewAliases)
 {
 	if (ensure(Asset.IsValid()))
 	{
 		const TArray<FName>* ExistingAliasesPtr = AliasesForObjectPath.Find(Asset.GetSoftObjectPath());
 		if (ExistingAliasesPtr)
 		{
-			TArray<FName> AliasesOnlyInExisting, AliasesOnlyInNew;
+			TArray<FName> AliasesOnlyInExisting;
+			TArray<AliasType> AliasesOnlyInNew;
 			FContentBrowserUniqueAlias UniqueAlias = MakeTuple(Asset.GetSoftObjectPath(), NAME_None);
 			for (const FName Alias : *ExistingAliasesPtr)
 			{
 				UniqueAlias.Value = Alias;
 				if (!AllAliases[UniqueAlias].bIsFromMetaData)
 				{
-					if (!NewAliases.Contains(Alias))
+					if (!NewAliases.ContainsByPredicate([&Alias](const AliasType& NewAlias) { return ContentBrowserAliasDataSource::GetAliasName(NewAlias) == Alias; }))
 					{
 						AliasesOnlyInExisting.Add(Alias);
 					}
 				}
 			}
 
-			for (const FName Alias : NewAliases)
+			for (const AliasType& Alias : NewAliases)
 			{
-				if (!ExistingAliasesPtr->Contains(Alias))
+				if (!ExistingAliasesPtr->Contains(ContentBrowserAliasDataSource::GetAliasName(Alias)))
 				{
 					AliasesOnlyInNew.Add(Alias);
 				}
@@ -900,7 +994,7 @@ FContentBrowserItemData UContentBrowserAliasDataSource::CreateAssetFolderItem(co
 
 	const FString FolderItemName = FPackageName::GetShortName(InFolderPath);
 	return FContentBrowserItemData(this, EContentBrowserItemFlags::Type_Folder | EContentBrowserItemFlags::Category_Asset, VirtualizedPath,
-		*FolderItemName, FText::GetEmpty(), MakeShared<FContentBrowserAssetFolderItemDataPayload>(InFolderPath));
+		*FolderItemName, AliasFolderDisplayNames.FindRef(InFolderPath), MakeShared<FContentBrowserAssetFolderItemDataPayload>(InFolderPath));
 }
 
 FContentBrowserItemData UContentBrowserAliasDataSource::CreateAssetFileItem(const FContentBrowserUniqueAlias& Alias)
@@ -918,7 +1012,7 @@ PRAGMA_ENABLE_DEPRECATION_WARNINGS
 		// Since AliasID is PackagePath/AssetName, AssetName should also be passed as the ItemName here. This provides the functionality of
 		// being able to have multiple aliases with the same display name, while still showing their original asset name in the tooltip.
 		return FContentBrowserItemData(this, EContentBrowserItemFlags::Type_File | EContentBrowserItemFlags::Category_Asset, VirtualizedPath,
-			AliasData->AssetData.AssetName, FText::FromName(AliasData->AliasName), MakeShared<FContentBrowserAliasItemDataPayload>(AliasData->AssetData, Alias));
+			AliasData->AssetData.AssetName, AliasData->AliasDisplayName, MakeShared<FContentBrowserAliasItemDataPayload>(AliasData->AssetData, Alias));
 	}
 	return FContentBrowserItemData();
 }
diff --git a/Engine/Plugins/Editor/ContentBrowser/ContentBrowserAliasDataSoure/Source/ContentBrowserAliasDataSource/Public/ContentBrowserAliasDataSource.h b/Engine/Plugins/Editor/ContentBrowser/ContentBrowserAliasDataSoure/Source/ContentBrowserAliasDataSource/Public/ContentBrowserAliasDataSource.h
--- a/Engine/Plugins/Editor/ContentBrowser/ContentBrowserAliasDataSoure/Source/ContentBrowserAliasDataSource/Public/ContentBrowserAliasDataSource.h
+++ b/Engine/Plugins/Editor/ContentBrowser/ContentBrowserAliasDataSoure/Source/ContentBrowserAliasDataSource/Public/ContentBrowserAliasDataSource.h
@@ -6,6 +6,7 @@
 #include "ContentBrowserAssetDataPayload.h"
 #include "ContentBrowserAssetDataSource.h"
 #include "ContentBrowserDataSource.h"
+#include "ContentBrowserLocalizedAlias.h"
 
 #include "AssetRegistry/PathTree.h"
 #include "ContentBrowserAliasDataSource.generated.h"
@@ -112,8 +113,10 @@ public:
 	 * @param bSkipPrimaryAssetValidation 	Use this if the asset being added is not a primary asset/is a re-director but should still be included in the Content Browser. (e.g. Verse classes)
 	 */
 	void AddAliases(const FAssetData& Asset, const TArray<FName>& Aliases, const bool bInIsFromMetaData = false, const bool bSkipPrimaryAssetValidation = false);
+	void AddAliases(const FAssetData& Asset, const TArray<FContentBrowserLocalizedAlias>& Aliases, const bool bInIsFromMetaData = false, const bool bSkipPrimaryAssetValidation = false);
 	/** Add an alias for a given asset. bInIsFromMetaData should only be true if the alias came from the AliasTagName metadata. */
 	void AddAlias(const FAssetData& Asset, const FName Alias, const bool bInIsFromMetaData = false, const bool bSkipPrimaryAssetValidation = false);
+	void AddAlias(const FAssetData& Asset, const FContentBrowserLocalizedAlias& Alias, const bool bInIsFromMetaData = false, const bool bSkipPrimaryAssetValidation = false);
 	/** Remove the given alias from the data source */
 	void RemoveAlias(const FSoftObjectPath& ObjectPath, const FName Alias);
 	/** Remove all aliases for the given object */
@@ -126,6 +129,16 @@ public:
 	UE_DEPRECATED(5.1, "FNames containing full asset paths are deprecated, use FSoftObjectPath instead")
 	void RemoveAliases(const FName ObjectPath);
 
+	/**
+	 * Add a display name override for the given alias folder, eg /MyAliases.
+	 * @note To provide a display name override for an alias itself, use the overloads that take a FContentBrowserLocalizedAlias.
+	 */
+	void AddAliasFolderDisplayName(const FName AliasFolder, const FText& DisplayName);
+	/**
+	 * Remove a display name override for the given alias folder, eg /MyAliases.
+	 */
+	void RemoveAliasFolderDisplayName(const FName AliasFolder);
+
 	/** When called, removes all aliases and triggers delegate for various systems to re-add aliases */
 	void RebuildAliases();
 
@@ -136,6 +149,7 @@ public:
 	void ReconcileAliasesFromMetaData(const FAssetData& Asset);
 	/** Calls AddAlias or RemoveAlias for every alias that doesn't match the stored data for the given asset. */
 	void ReconcileAliasesForAsset(const FAssetData& Asset, const TArray<FName>& NewAliases);
+	void ReconcileAliasesForAsset(const FAssetData& Asset, const TArray<FContentBrowserLocalizedAlias>& NewAliases);
 
 	/** Logs all the content browser aliases */
 	void LogAliases() const;
@@ -144,6 +158,13 @@ protected:
 	virtual void BuildRootPathVirtualTree() override;
 
 private:
+	template <typename AliasType>
+	void AddAliasImpl(const FAssetData& Asset, const AliasType& Alias, const bool bInIsFromMetaData, const bool bSkipPrimaryAssetValidation);
+	template <typename AliasType>
+	void AddAliasesImpl(const FAssetData& Asset, const TArray<AliasType>& Aliases, const bool bInIsFromMetaData, const bool bSkipPrimaryAssetValidation);
+	template <typename AliasType>
+	void ReconcileAliasesForAssetImpl(const FAssetData& Asset, const TArray<AliasType>& NewAliases);
+
 	void OnAssetAdded(const FAssetData& InAssetData);
 	void OnAssetRemoved(const FAssetData& InAssetData);
 	void OnAssetUpdated(const FAssetData& InAssetData);
@@ -165,8 +186,8 @@ private:
 	struct FAliasData
 	{
 		FAliasData() {}
-		FAliasData(const FAssetData& InAssetData, const FName InPackagePath, const FName InName, const bool bInIsFromMetaData = false)
-			: AssetData(InAssetData), PackagePath(InPackagePath), AliasName(InName), bIsFromMetaData(bInIsFromMetaData)
+		FAliasData(const FAssetData& InAssetData, const FName InPackagePath, const FText& InDisplayName, const bool bInIsFromMetaData = false)
+			: AssetData(InAssetData), PackagePath(InPackagePath), AliasDisplayName(InDisplayName), bIsFromMetaData(bInIsFromMetaData)
 		{
 			FNameBuilder AssetNameBuilder(InAssetData.AssetName);
 			{
@@ -194,7 +215,7 @@ private:
 		/** PackageName.SourceAssetName, /MyAliases/SomeAsset.SomeAsset */
 		FSoftObjectPath ObjectPath;
 		/** A non-unique display name for this alias */
-		FName AliasName;
+		FText AliasDisplayName;
 		/** Whether this alias was generated from package metadata or manually through the C++ interface */
 		bool bIsFromMetaData = false;
 	};
@@ -202,12 +223,14 @@ private:
 
 	/** The full folder hierarchy for all alias paths */
 	FPathTree PathTree;
-	/** Alias data keyed by their full alias path, ie /Game/MyData/Aliases/SourceMesh */
+	/** Alias data keyed by their full alias path, eg /Game/MyData/Aliases/SourceMesh */
 	TMap<FContentBrowserUniqueAlias, FAliasData> AllAliases;
-	/** A list of alias paths to display for each asset, ie /Game/Meshes/SourceMesh.SourceMesh */
+	/** A list of alias paths to display for each asset, eg /Game/Meshes/SourceMesh.SourceMesh */
 	TMap<FSoftObjectPath, TArray<FName>> AliasesForObjectPath;
-	/** A list of alias paths to display for each folder, ie /Game/MyData/Aliases */
+	/** A list of alias paths to display for each folder, eg /Game/MyData/Aliases */
 	TMap<FName, TArray<FContentBrowserUniqueAlias>> AliasesInPackagePath;
+	/** Alias folder display names keyed against their alias path, eg /Game/MyData/Aliases/ */
+	TMap<FName, FText> AliasFolderDisplayNames;
 	/** A set used for removing duplicate aliases in the same query, stored here to avoid constant reallocation */
 	TSet<FSoftObjectPath> AlreadyAddedOriginalAssets;
 
diff --git a/Engine/Plugins/Editor/ContentBrowser/ContentBrowserAliasDataSoure/Source/ContentBrowserAliasDataSource/Public/ContentBrowserLocalizedAlias.h b/Engine/Plugins/Editor/ContentBrowser/ContentBrowserAliasDataSoure/Source/ContentBrowserAliasDataSource/Public/ContentBrowserLocalizedAlias.h
new file mode 100644
--- /dev/null
+++ b/Engine/Plugins/Editor/ContentBrowser/ContentBrowserAliasDataSoure/Source/ContentBrowserAliasDataSource/Public/ContentBrowserLocalizedAlias.h
@@ -0,0 +1,28 @@
+// Copyright Epic Games, Inc. All Rights Reserved.
+
+#pragma once
+
+#include "Internationalization/Text.h"
+#include "UObject/NameTypes.h"
+
+/** Container for an alias and its optional display name override */
+struct FContentBrowserLocalizedAlias
+{
+	FName Alias;
+	FText DisplayName;
+
+	bool operator==(const FContentBrowserLocalizedAlias& Other) const
+	{
+		return Alias == Other.Alias;
+	}
+
+	bool operator!=(const FContentBrowserLocalizedAlias& Other) const
+	{
+		return Alias != Other.Alias;
+	}
+
+	friend uint32 GetTypeHash(const FContentBrowserLocalizedAlias& Alias)
+	{
+		return GetTypeHash(Alias.Alias);
+	}
+};
