Fixed Clang compile error

#rnx

diff --git a/Engine/Plugins/Editor/ContentBrowser/ContentBrowserAliasDataSoure/Source/ContentBrowserAliasDataSource/Public/ContentBrowserLocalizedAlias.h b/Engine/Plugins/Editor/ContentBrowser/ContentBrowserAliasDataSoure/Source/ContentBrowserAliasDataSource/Public/ContentBrowserLocalizedAlias.h
--- a/Engine/Plugins/Editor/ContentBrowser/ContentBrowserAliasDataSoure/Source/ContentBrowserAliasDataSource/Public/ContentBrowserLocalizedAlias.h
+++ b/Engine/Plugins/Editor/ContentBrowser/ContentBrowserAliasDataSoure/Source/ContentBrowserAliasDataSource/Public/ContentBrowserLocalizedAlias.h
@@ -21,8 +21,8 @@ struct FContentBrowserLocalizedAlias
 		return Alias != Other.Alias;
 	}
 
-	friend uint32 GetTypeHash(const FContentBrowserLocalizedAlias& Alias)
+	friend uint32 GetTypeHash(const FContentBrowserLocalizedAlias& This)
 	{
-		return GetTypeHash(Alias.Alias);
+		return GetTypeHash(This.Alias);
 	}
 };
